Asset Search - Working on a new plugin to make it possible to index lots of differnet types of assets and search across a wide swath of them to find information such as text inside a data table.

Not done yet, still marking it as beta.  Not enabled by default.

diff --git a/Engine/Plugins/Editor/AssetSearch/AssetSearch.uplugin b/Engine/Plugins/Editor/AssetSearch/AssetSearch.uplugin
new file mode 100644
--- /dev/null
+++ b/Engine/Plugins/Editor/AssetSearch/AssetSearch.uplugin
@@ -0,0 +1,30 @@
+{
+	"FileVersion": 3,
+	"Version": 1,
+	"VersionName": "1.0",
+	"FriendlyName": "Asset Search",
+	"Description": "",
+	"Category": "Editor",
+	"CreatedBy": "Epic Games, Inc.",
+	"CreatedByURL": "http://epicgames.com",
+	"DocsURL": "",
+	"MarketplaceURL": "",
+	"SupportURL": "",
+	"EnabledByDefault": false,
+	"CanContainContent": false,
+	"IsBetaVersion": true,
+	"Installed": false,
+	"Modules": [
+		{
+			"Name": "AssetSearch",
+			"Type": "Editor",
+			"LoadingPhase": "Default"
+		}
+	],
+	"Plugins": [
+		{
+			"Name": "SQLiteCore",
+			"Enabled": true
+		}
+	]
+}
\ No newline at end of file
diff --git a/Engine/Plugins/Editor/AssetSearch/Source/AssetSearch.Build.cs b/Engine/Plugins/Editor/AssetSearch/Source/AssetSearch.Build.cs
new file mode 100644
--- /dev/null
+++ b/Engine/Plugins/Editor/AssetSearch/Source/AssetSearch.Build.cs
@@ -0,0 +1,43 @@
+// Copyright Epic Games, Inc. All Rights Reserved.
+
+using UnrealBuildTool;
+
+public class AssetSearch : ModuleRules
+{
+	public AssetSearch(ReadOnlyTargetRules Target) : base(Target)
+	{
+		PublicDependencyModuleNames.AddRange(
+			new string[]
+			{
+				"Core",
+				"CoreUObject",
+				"Engine",
+				"Json",
+			}
+		);
+
+		PrivateDependencyModuleNames.AddRange(
+			new string[] {
+				"Core",
+				"CoreUObject",
+				"Engine",
+				"SQLiteCore",
+				"EditorStyle",
+				"InputCore",
+				"Slate",
+				"SlateCore",
+				"GameplayTags",
+				"DerivedDataCache",
+				"WorkspaceMenuStructure",
+				"UnrealEd",
+				"Json",
+				"AssetRegistry",
+				"JsonUtilities",
+
+				"UMG",
+				"UMGEditor",
+				"BlueprintGraph",
+			}
+		);
+	}
+}
diff --git a/Engine/Plugins/Editor/AssetSearch/Source/Private/AssetSearchDatabase.cpp b/Engine/Plugins/Editor/AssetSearch/Source/Private/AssetSearchDatabase.cpp
new file mode 100644
--- /dev/null
+++ b/Engine/Plugins/Editor/AssetSearch/Source/Private/AssetSearchDatabase.cpp
@@ -0,0 +1,694 @@
+// Copyright Epic Games, Inc. All Rights Reserved.
+
+#include "AssetSearchDatabase.h"
+
+#include "SQLiteDatabase.h"
+#include "HAL/FileManager.h"
+#include "UObject/StructOnScope.h"
+#include "Serialization/MemoryReader.h"
+#include "Serialization/MemoryWriter.h"
+
+#include "AssetData.h"
+#include "Misc/TextFilterExpressionEvaluator.h"
+
+PRAGMA_DISABLE_OPTIMIZATION
+
+DECLARE_LOG_CATEGORY_CLASS(LogAssetSearch, Log, All);
+
+enum class EAssetSearchDatabaseVersion
+{
+	Empty = 0,
+	Initial = 1,
+
+	Current = Initial,
+};
+
+class FAssetSearchDatabaseStatements
+{
+private:
+	FSQLiteDatabase& Database;
+
+public:
+	explicit FAssetSearchDatabaseStatements(FSQLiteDatabase& InDatabase)
+		: Database(InDatabase)
+	{
+		check(Database.IsValid());
+	}
+
+	bool CreatePreparedStatements()
+	{
+		check(Database.IsValid());
+
+#define PREPARE_STATEMENT(VAR)																				\
+			(VAR) = Database.PrepareStatement<decltype(VAR)>(ESQLitePreparedStatementFlags::Persistent);	\
+			if (!(VAR).IsValid()) { return false; }
+
+		PREPARE_STATEMENT(Statement_BeginTransaction);
+		PREPARE_STATEMENT(Statement_CommitTransaction);
+		PREPARE_STATEMENT(Statement_RollbackTransaction);
+
+		PREPARE_STATEMENT(Statement_GetAssetIdForAssetPath);
+		PREPARE_STATEMENT(Statement_IsAssetUpToDate);
+		PREPARE_STATEMENT(Statement_AddAssetToAssetTable);
+		PREPARE_STATEMENT(Statement_AddAssetProperty);
+		PREPARE_STATEMENT(Statement_DeleteEntriesForAsset);
+
+		PREPARE_STATEMENT(Statement_SearchAssetsFTS);
+
+#undef PREPARE_STATEMENT
+
+		return true;
+	}
+
+	/**
+	 * Statements managing database transactions
+	 */
+
+	 /** Begin a database transaction */
+	SQLITE_PREPARED_STATEMENT_SIMPLE(FBeginTransaction, "BEGIN TRANSACTION;");
+	FBeginTransaction Statement_BeginTransaction;
+	bool BeginTransaction()
+	{
+		return Statement_BeginTransaction.Execute();
+	}
+
+	/** Commit a database transaction */
+	SQLITE_PREPARED_STATEMENT_SIMPLE(FCommitTransaction, "COMMIT TRANSACTION;");
+	FCommitTransaction Statement_CommitTransaction;
+	bool CommitTransaction()
+	{
+		return Statement_CommitTransaction.Execute();
+	}
+
+	/** Rollback a database transaction */
+	SQLITE_PREPARED_STATEMENT_SIMPLE(FRollbackTransaction, "ROLLBACK TRANSACTION;");
+	FRollbackTransaction Statement_RollbackTransaction;
+	bool RollbackTransaction()
+	{
+		return Statement_RollbackTransaction.Execute();
+	}
+
+	/**
+	 * Application Statements
+	 */
+
+	SQLITE_PREPARED_STATEMENT(FIsAssetUpToDate,
+		"SELECT index_hash FROM table_assets WHERE asset_path = ?1;",
+		SQLITE_PREPARED_STATEMENT_COLUMNS(FString),
+		SQLITE_PREPARED_STATEMENT_BINDINGS(FString)
+	);
+	private: FIsAssetUpToDate Statement_IsAssetUpToDate;
+	public: bool IsAssetUpToDate(const FAssetData& InAssetData, const FString& IndexedJsonHash)
+	{
+		FString OutIndexedJsonHash;
+		if (Statement_IsAssetUpToDate.BindAndExecuteSingle(InAssetData.ObjectPath.ToString(), OutIndexedJsonHash))
+		{
+			return OutIndexedJsonHash.Equals(IndexedJsonHash, ESearchCase::CaseSensitive);
+		}
+
+		return false;
+	}
+
+	SQLITE_PREPARED_STATEMENT(FGetAssetIdForAssetPath,
+		"SELECT assetid FROM table_assets WHERE asset_path = ?1;",
+		SQLITE_PREPARED_STATEMENT_COLUMNS(int64),
+		SQLITE_PREPARED_STATEMENT_BINDINGS(FString)
+	);
+	private: FGetAssetIdForAssetPath Statement_GetAssetIdForAssetPath;
+	public: int64 GetAssetIdForAsset(const FAssetData& InAssetData)
+	{
+		int64 OutAssetId = INDEX_NONE;
+		if (Statement_GetAssetIdForAssetPath.BindAndExecuteSingle(InAssetData.ObjectPath.ToString(), OutAssetId))
+		{
+			return OutAssetId;
+		}
+
+		return INDEX_NONE;
+	}
+
+	SQLITE_PREPARED_STATEMENT_BINDINGS_ONLY(
+		FAddAssetToAssetTable,
+		" INSERT INTO table_assets(asset_name, asset_class, asset_path, index_hash)"
+		" VALUES(?1, ?2, ?3, ?4);",
+		SQLITE_PREPARED_STATEMENT_BINDINGS(FString, FString, FString, FString)
+	);
+	private: FAddAssetToAssetTable Statement_AddAssetToAssetTable;
+	SQLITE_PREPARED_STATEMENT_BINDINGS_ONLY(
+		FAddAssetPropertiesFromJson,
+		" INSERT INTO table_asset_properties(assetid, object_name, object_path, object_native_class, property_name, property_field, property_class, value_text, value_hidden)"
+		" VALUES(?1, ?2, ?3, ?4, ?5, ?6, ?7, ?8, ?9);",
+		SQLITE_PREPARED_STATEMENT_BINDINGS(int64, FString, FString, FString, FString, FString, FString, FString, FString)
+	);
+	private: FAddAssetPropertiesFromJson Statement_AddAssetProperty;
+	public: bool AddSearchRecord(const FAssetData& InAssetData, const FString& IndexedJson, const FString& IndexedJsonHash)
+	{
+		if (Statement_AddAssetToAssetTable.BindAndExecute(InAssetData.AssetName.ToString(), InAssetData.AssetClass.ToString(), InAssetData.ObjectPath.ToString(), IndexedJsonHash))
+		{
+			int64 AssetId = Database.GetLastInsertRowId();
+			//ensure(AssetId == GetAssetIdForAsset(InAssetData));
+
+			TSharedRef<TJsonReader<>> JsonReader = TJsonReaderFactory<>::Create(IndexedJson);
+
+			EJsonNotation Notation;
+
+			ensure(JsonReader->ReadNext(Notation) && Notation == EJsonNotation::ObjectStart);
+
+			ensure(JsonReader->ReadNext(Notation) && Notation == EJsonNotation::Number && JsonReader->GetIdentifier() == TEXT("version"));
+			const int32 SerializerVersion = (int32)JsonReader->GetValueAsNumber();
+
+			ensure(JsonReader->ReadNext(Notation) && Notation == EJsonNotation::ObjectStart && JsonReader->GetIdentifier() == TEXT("indexers"));
+
+			while (JsonReader->ReadNext(Notation))
+			{
+				// Indexer
+				if (Notation == EJsonNotation::ObjectStart)
+				{
+					const FString IndexerName = JsonReader->GetIdentifier();
+
+					ensure(JsonReader->ReadNext(Notation) && Notation == EJsonNotation::Number && JsonReader->GetIdentifier() == TEXT("version"));
+					const int32 IndexerVersion = (int32)JsonReader->GetValueAsNumber();
+
+					ensure(JsonReader->ReadNext(Notation) && Notation == EJsonNotation::ArrayStart && JsonReader->GetIdentifier() == TEXT("objects"));
+
+					// Objects[]
+					while (JsonReader->ReadNext(Notation) && Notation != EJsonNotation::ArrayEnd)
+					{
+						if (Notation == EJsonNotation::ObjectStart)
+						{
+							FString object_name;
+							FString object_path;
+							FString object_native_class;
+
+							ensure(JsonReader->ReadNext(Notation) && Notation == EJsonNotation::String && JsonReader->GetIdentifier() == TEXT("name"));
+							object_name = JsonReader->GetValueAsString();
+							ensure(JsonReader->ReadNext(Notation) && Notation == EJsonNotation::String && JsonReader->GetIdentifier() == TEXT("path"));
+							object_path = JsonReader->GetValueAsString();
+							ensure(JsonReader->ReadNext(Notation) && Notation == EJsonNotation::String && JsonReader->GetIdentifier() == TEXT("native_class"));
+							object_native_class = JsonReader->GetValueAsString();
+
+							ensure(JsonReader->ReadNext(Notation) && Notation == EJsonNotation::ArrayStart && JsonReader->GetIdentifier() == TEXT("properties"));
+
+							FString property_name;
+							FString property_field;
+							FString property_class;
+							FString value_text;
+							FString value_hidden;
+
+							// Begin Properties[]
+							while (JsonReader->ReadNext(Notation) && Notation != EJsonNotation::ArrayEnd)
+							{
+								if (Notation == EJsonNotation::ObjectStart)
+								{
+									// Read all attributes of a property
+									while (JsonReader->ReadNext(Notation) && Notation != EJsonNotation::ObjectEnd)
+									{
+										if (Notation == EJsonNotation::String)
+										{
+											if (JsonReader->GetIdentifier() == TEXT("name"))
+											{
+												property_name = JsonReader->GetValueAsString();
+											}
+											else if (JsonReader->GetIdentifier() == TEXT("field"))
+											{
+												property_field = JsonReader->GetValueAsString();
+											}
+											else if (JsonReader->GetIdentifier() == TEXT("class"))
+											{
+												property_class = JsonReader->GetValueAsString();
+											}
+											else if (JsonReader->GetIdentifier() == TEXT("value_text"))
+											{
+												value_text = JsonReader->GetValueAsString();
+											}
+											else if (JsonReader->GetIdentifier() == TEXT("value_hidden"))
+											{
+												value_hidden = JsonReader->GetValueAsString();
+											}
+										}
+									}
+
+									if (!Statement_AddAssetProperty.BindAndExecute(AssetId, object_name, object_path, object_native_class, property_name, property_field, property_class, value_text, value_hidden))
+									{
+										//Log error?
+									}
+								}
+							}
+							// End Properties[]
+						}
+					}
+					// End Objects[]
+				}
+			}
+
+			return true;
+		}
+		return false;
+	}
+
+	SQLITE_PREPARED_STATEMENT_BINDINGS_ONLY(FDeleteEntriesForAsset, "DELETE FROM table_assets WHERE asset_path = ?1;", SQLITE_PREPARED_STATEMENT_BINDINGS(FString));
+	private: FDeleteEntriesForAsset Statement_DeleteEntriesForAsset;
+	public: bool DeleteEntriesForAsset(const FAssetData& InAssetData)
+	{
+		if (Statement_DeleteEntriesForAsset.BindAndExecute(InAssetData.ObjectPath.ToString()))
+		{
+			return true;
+		}
+		return false;
+	}
+
+	SQLITE_PREPARED_STATEMENT(FSearchAssetsFTS,
+		" SELECT "
+		"     asset_name, "
+		"     asset_class, "
+		"     asset_path, "
+		"     object_name, "
+		"     object_path, "
+		"     object_native_class, "
+		"     property_name, "
+		"     property_field, "
+		"     property_class, "
+		"     value_text, "
+		"     value_hidden, "
+		"     rank as score "
+		" FROM table_asset_properties_fts "
+		" WHERE table_asset_properties_fts MATCH ?1 "
+		//" ORDER BY rank "
+		";",
+		SQLITE_PREPARED_STATEMENT_COLUMNS(FString, FString, FString, FString, FString, FString, FString, FString, FString, FString, FString, float),
+		SQLITE_PREPARED_STATEMENT_BINDINGS(FString)
+	);
+	FSearchAssetsFTS Statement_SearchAssetsFTS;
+	bool SearchAssets(const FSearchQuery& Query, TFunctionRef<ESQLitePreparedStatementExecuteRowResult(FSearchRecord&&)> InCallback)
+	{
+		//FString BaseSearch(
+		//	" SELECT "
+		//	"     asset_name, "
+		//	"     asset_class, "
+		//	"     asset_path, "
+		//	"     object_name, "
+		//	"     object_path, "
+		//	"     object_native_class, "
+		//	"     property_name, "
+		//	"     property_field, "
+		//	"     property_class, "
+		//	"     value_text, "
+		//	"     value_hidden "
+		//	" FROM view_asset_properties "
+		//	" WHERE ");
+
+		//const bool bEvaledFilter = Eval.SetFilterText(FText::FromString(Query.Query));
+		//if (!ensure(bEvaledFilter))
+		//{
+		//	return false;
+		//}
+
+		//FTextFilterExpressionEvaluator Eval(ETextFilterExpressionEvaluatorMode::BasicString);
+		//const TArray<FExpressionToken>& Tokens = Eval.GetFilterExpressionTokens();
+
+		//for (int32 i = 0; i < Tokens.Num(); i++)
+		//{
+		//	const FStringToken& Token = Tokens[i].Context;
+		//	
+		//	if (Tokens[i].Node.Cast<TextFilterExpressionParser::FEqual>())
+		//	{
+
+		//	}
+		//}
+
+
+		//FSQLitePreparedStatement DynamicSearchStatement = Database.PrepareStatement(ESQLitePreparedStatementFlags::None);
+		
+
+		FString Q;
+
+		FTextFilterExpressionEvaluator Eval(ETextFilterExpressionEvaluatorMode::BasicString);
+		const TArray<FExpressionToken>& Tokens = Eval.GetFilterExpressionTokens();
+		if (Eval.SetFilterText(FText::FromString(Query.Query)))
+		{
+			for (int32 i = 0; i < Tokens.Num(); i++)
+			{
+				const FStringToken& Token = Tokens[i].Context;
+
+				if (Token.GetString().StartsWith(TEXT("\"")) && Token.GetString().EndsWith(TEXT("\"")))
+				{
+					Q += Token.GetString();
+					Q += TEXT(" ");
+				}
+				else
+				{
+					Q += TEXT("\"") + Token.GetString() + TEXT("\" * ");
+				}
+			}
+		}
+		else
+		{
+			TArray<FString> Phrases;
+			Query.Query.ParseIntoArray(Phrases, TEXT(" "), 1);
+
+			for (FString Phrase : Phrases)
+			{
+				Q += TEXT("\"") + Phrase + TEXT("\" * ");
+			}
+		}
+
+		return Statement_SearchAssetsFTS.BindAndExecute(Q, [&InCallback](const FSearchAssetsFTS& InStatement)
+		{
+			FSearchRecord Result;
+			if (InStatement.GetColumnValues(
+				Result.AssetName, Result.AssetClass, Result.AssetPath,
+				Result.object_name, Result.object_path, Result.object_native_class,
+				Result.property_name, Result.property_field, Result.property_class,
+				Result.value_text, Result.value_hidden,
+				Result.Score))
+			{
+				return InCallback(MoveTemp(Result));
+			}
+			return ESQLitePreparedStatementExecuteRowResult::Error;
+		}) != INDEX_NONE;
+	}
+};
+
+class FAssetSearchDatabaseScopedTransaction
+{
+public:
+	explicit FAssetSearchDatabaseScopedTransaction(FAssetSearchDatabaseStatements& InStatements)
+		: Statements(InStatements)
+		, bHasTransaction(Statements.BeginTransaction()) // This will fail if a transaction is already open
+	{
+	}
+
+	~FAssetSearchDatabaseScopedTransaction()
+	{
+		Commit();
+	}
+
+	bool CommitOrRollback(const bool bShouldCommit)
+	{
+		if (bShouldCommit)
+		{
+			Commit();
+			return true;
+		}
+
+		Rollback();
+		return false;
+	}
+
+	void Commit()
+	{
+		if (bHasTransaction)
+		{
+			verify(Statements.CommitTransaction());
+			bHasTransaction = false;
+		}
+	}
+
+	void Rollback()
+	{
+		if (bHasTransaction)
+		{
+			verify(Statements.RollbackTransaction());
+			bHasTransaction = false;
+		}
+	}
+
+private:
+	FAssetSearchDatabaseStatements& Statements;
+	bool bHasTransaction;
+};
+
+FAssetSearchDatabase::FAssetSearchDatabase()
+	: Database(MakeUnique<FSQLiteDatabase>())
+	, DatabaseFileName(TEXT("AssetSearch.db"))
+{
+}
+
+FAssetSearchDatabase::~FAssetSearchDatabase()
+{
+	Close();
+}
+
+bool FAssetSearchDatabase::IsValid() const
+{
+	return Database->IsValid();
+}
+
+bool FAssetSearchDatabase::Open(const FString& InSessionPath)
+{
+	return Open(InSessionPath, ESQLiteDatabaseOpenMode::ReadWriteCreate);
+}
+
+bool FAssetSearchDatabase::Open(const FString& InSessionPath, const ESQLiteDatabaseOpenMode InOpenMode)
+{
+	if (Database->IsValid())
+	{
+		return false;
+	}
+
+	if (!Database->Open(*(InSessionPath / DatabaseFileName), ESQLiteDatabaseOpenMode::ReadWriteCreate))
+	{
+		UE_LOG(LogAssetSearch, Error, TEXT("Failed to open database for '%s': %s"), *InSessionPath, *GetLastError());
+		return false;
+	}
+
+	SessionPath = InSessionPath;
+
+	// Set the database to use exclusive WAL mode for performance (exclusive works even on platforms without a mmap implementation)
+	// Set the database "NORMAL" fsync mode to only perform a fsync when check-pointing the WAL to the main database file (fewer fsync calls are better for performance, with a very slight loss of WAL durability if the power fails)
+	Database->Execute(TEXT("PRAGMA locking_mode=EXCLUSIVE;"));
+	Database->Execute(TEXT("PRAGMA journal_mode=WAL;"));
+	Database->Execute(TEXT("PRAGMA synchronous=NORMAL;"));
+
+	int32 LoadedDatabaseVersion = 0;
+	Database->GetUserVersion(LoadedDatabaseVersion);
+	if (LoadedDatabaseVersion != (int32)EAssetSearchDatabaseVersion::Empty)
+	{
+		if (LoadedDatabaseVersion > (int32)EAssetSearchDatabaseVersion::Current)
+		{
+			Close();
+			UE_LOG(LogAssetSearch, Error, TEXT("Failed to open database for '%s': Database is too new (version %d, expected = %d)"), *InSessionPath, LoadedDatabaseVersion, (int32)EAssetSearchDatabaseVersion::Current);
+			return false;
+		}
+		else if (LoadedDatabaseVersion < (int32)EAssetSearchDatabaseVersion::Current)
+		{
+			Close(true);
+			UE_LOG(LogAssetSearch, Log, TEXT("Opened database '%s': Database is too old (version %d, expected = %d), creating new database"), *InSessionPath, LoadedDatabaseVersion, (int32)EAssetSearchDatabaseVersion::Current);
+			return Open(InSessionPath, InOpenMode);
+		}
+	}
+
+	// Create our required tables
+	//========================================================================
+	if (!ensure(Database->Execute(TEXT("CREATE TABLE IF NOT EXISTS table_assets(assetid INTEGER PRIMARY KEY, asset_name, asset_class, asset_path, index_hash);"))))
+	{
+		LogLastError();
+		Close();
+		return false;
+	}
+
+	if (!ensure(Database->Execute(TEXT("CREATE TABLE IF NOT EXISTS table_asset_properties(rowid INTEGER PRIMARY KEY, object_name, object_path, object_native_class, property_name, property_field, property_class, value_text, value_hidden, assetid INTEGER, FOREIGN KEY(assetid) REFERENCES table_assets(assetid));"))))
+	{
+		LogLastError();
+		Close();
+		return false;
+	}
+
+	if (!ensure(Database->Execute(TEXT("CREATE VIEW IF NOT EXISTS view_asset_properties AS SELECT rowid, object_name, object_path, object_native_class, property_name, property_field, property_class, value_text, value_hidden, table_asset_properties.assetid as assetid, table_assets.asset_name AS asset_name, table_assets.asset_class AS asset_class, table_assets.asset_path AS asset_path, table_assets.index_hash AS index_hash FROM table_asset_properties INNER JOIN table_assets on table_assets.assetid = table_asset_properties.assetid;"))))
+	{
+		LogLastError();
+		Close();
+		return false;
+	}
+
+	if (!ensure(Database->Execute(TEXT("CREATE VIRTUAL TABLE IF NOT EXISTS table_asset_properties_fts USING FTS5(asset_name, asset_class UNINDEXED, asset_path UNINDEXED, object_name UNINDEXED, object_path UNINDEXED, object_native_class UNINDEXED, property_name UNINDEXED, property_field UNINDEXED, property_class UNINDEXED, value_text, value_hidden, assetid UNINDEXED, content=view_asset_properties, content_rowid=rowid);"))))
+	{
+		LogLastError();
+		Close();
+		return false;
+	}
+
+	if (!ensure(Database->Execute(TEXT(" DROP TRIGGER IF EXISTS table_assets_delete;"))))
+	{
+		LogLastError();
+		Close();
+		return false;
+	}
+
+	if (!ensure(Database->Execute(
+		TEXT(" CREATE TRIGGER table_assets_delete BEFORE DELETE ON table_assets BEGIN")
+		TEXT("     DELETE FROM table_asset_properties WHERE assetid == old.assetid;")
+		TEXT(" END;")
+	)))
+	{
+		LogLastError();
+		Close();
+		return false;
+	}
+
+	if (!ensure(Database->Execute(TEXT(" DROP TRIGGER IF EXISTS table_asset_properties_insert;"))))
+	{
+		LogLastError();
+		Close();
+		return false;
+	}
+
+	if (!ensure(Database->Execute(
+		TEXT(" CREATE TRIGGER table_asset_properties_insert AFTER INSERT ON table_asset_properties BEGIN")
+		TEXT("     INSERT INTO table_asset_properties_fts(rowid, object_name, object_path, object_native_class, property_name, property_field, property_class, value_text, value_hidden, assetid) VALUES (new.rowid, new.object_name, new.object_path, new.object_native_class, new.property_name, new.property_field, new.property_class, new.value_text, new.value_hidden, new.assetid);")
+		TEXT(" END;")
+		)))
+	{
+		LogLastError();
+		Close();
+	}
+
+	if (!ensure(Database->Execute(
+		TEXT(" DROP TRIGGER IF EXISTS table_asset_properties_delete;")
+	)))
+	{
+		LogLastError();
+		Close();
+		return false;
+	}
+
+	if (!ensure(Database->Execute(
+		TEXT(" CREATE TRIGGER table_asset_properties_delete AFTER DELETE ON table_asset_properties BEGIN")
+		TEXT("     INSERT INTO table_asset_properties_fts(table_asset_properties_fts, rowid, object_name, object_path, object_native_class, property_name, property_field, property_class, value_text, value_hidden, assetid) VALUES('delete', old.rowid, old.object_name, old.object_path, old.object_native_class, old.property_name, old.property_field, old.property_class, old.value_text, old.value_hidden, old.assetid);")
+		TEXT(" END;")
+	)))
+	{
+		LogLastError();
+		Close();
+		return false;
+	}
+
+	if (!ensure(Database->Execute(
+		TEXT(" DROP TRIGGER IF EXISTS table_asset_properties_update;")
+	)))
+	{
+		LogLastError();
+		Close();
+		return false;
+	}
+
+	if (!ensure(Database->Execute(
+		TEXT(" CREATE TRIGGER table_asset_properties_update AFTER UPDATE ON table_asset_properties BEGIN")
+		TEXT("     INSERT INTO table_asset_properties_fts(table_asset_properties_fts, rowid, object_name, object_path, object_native_class, property_name, property_field, property_class, value_text, value_hidden, assetid) VALUES('delete', old.rowid, old.object_name, old.object_path, old.object_native_class, old.property_name, old.property_field, old.property_class, old.value_text, old.value_hidden, old.assetid);")
+		TEXT("     INSERT INTO table_asset_properties_fts(rowid, object_name, object_path, object_native_class, property_name, property_field, property_class, value_text, value_hidden, assetid) VALUES (new.rowid, new.object_name, new.object_path, new.object_native_class, new.property_name, new.property_field, new.property_class, new.value_text, new.value_hidden, new.assetid);")
+		TEXT(" END;")
+	)))
+	{
+		LogLastError();
+		Close();
+		return false;
+	}
+
+	if (!ensure(Database->Execute(
+		TEXT("CREATE UNIQUE INDEX IF NOT EXISTS asset_path_index ON table_assets(asset_path);")
+	)))
+	{
+		LogLastError();
+		Close();
+		return false;
+	}
+
+	//CREATE INDEX blueprint_nodes ON table_assets(json_extract(json, '$.Blueprint.Nodes'))
+
+	// The database will have the latest schema at this point, so update the user-version
+	if (!Database->SetUserVersion((int32)EAssetSearchDatabaseVersion::Current))
+	{
+		Close();
+		return false;
+	}
+
+	// Create our required prepared statements
+	Statements = MakeUnique<FAssetSearchDatabaseStatements>(*Database);
+	if (!ensure(Statements->CreatePreparedStatements()))
+	{
+		Close();
+		return false;
+	}
+
+	return true;
+}
+
+bool FAssetSearchDatabase::Close(const bool InDeleteDatabase)
+{
+	if (!Database->IsValid())
+	{
+		return false;
+	}
+
+	// Need to destroy prepared statements before the database can be closed
+	Statements.Reset();
+
+	if (!Database->Close())
+	{
+		UE_LOG(LogAssetSearch, Error, TEXT("Failed to close database for '%s': %s"), *SessionPath, *GetLastError());
+		return false;
+	}
+
+	if (InDeleteDatabase)
+	{
+		IFileManager::Get().Delete(*(SessionPath / DatabaseFileName), false);
+	}
+
+	SessionPath.Reset();
+
+	return true;
+}
+
+FString FAssetSearchDatabase::GetFilename() const
+{
+	return Database->GetFilename();
+}
+
+FString FAssetSearchDatabase::GetLastError() const
+{
+	return Database->GetLastError();
+}
+
+void FAssetSearchDatabase::LogLastError() const
+{
+	UE_LOG(LogAssetSearch, Error, TEXT("Database Error: %s"), *SessionPath, *GetLastError());
+}
+
+bool FAssetSearchDatabase::IsAssetUpToDate(const FAssetData& InAssetData, const FString& IndexedJsonHash)
+{
+	if (ensure(Statements))
+	{
+		return Statements->IsAssetUpToDate(InAssetData, IndexedJsonHash);
+	}
+
+	return false;
+}
+
+void FAssetSearchDatabase::AddOrUpdateAsset(const FAssetData& InAssetData, const FString& IndexedJson, const FString& IndexedJsonHash)
+{
+	if (ensure(Statements))
+	{
+		if (!ensure(Statements->DeleteEntriesForAsset(InAssetData)))
+		{
+			LogLastError();
+		}
+
+		if (!ensure(Statements->AddSearchRecord(InAssetData, IndexedJson, IndexedJsonHash)))
+		{
+			LogLastError();
+		}
+	}
+}
+
+bool FAssetSearchDatabase::EnumerateSearchResults(const FSearchQuery& Query, TFunctionRef<bool(FSearchRecord&&)> InCallback)
+{
+	bool bSuccess = Statements->SearchAssets(Query, [&InCallback](FSearchRecord&& InResult)
+	{
+		return InCallback(MoveTemp(InResult))
+			? ESQLitePreparedStatementExecuteRowResult::Continue
+			: ESQLitePreparedStatementExecuteRowResult::Stop;
+	});
+
+	return bSuccess;
+}
+
+PRAGMA_ENABLE_OPTIMIZATION
diff --git a/Engine/Plugins/Editor/AssetSearch/Source/Private/AssetSearchDatabase.h b/Engine/Plugins/Editor/AssetSearch/Source/Private/AssetSearchDatabase.h
new file mode 100644
--- /dev/null
+++ b/Engine/Plugins/Editor/AssetSearch/Source/Private/AssetSearchDatabase.h
@@ -0,0 +1,97 @@
+// Copyright Epic Games, Inc. All Rights Reserved.
+
+#pragma once
+
+#include "CoreMinimal.h"
+#include "IAssetSearchModule.h"
+
+class FSQLiteDatabase;
+enum class ESQLiteDatabaseOpenMode : uint8;
+class FAssetSearchDatabaseStatements;
+struct FAssetData;
+
+class FAssetSearchDatabase
+{
+public:
+	FAssetSearchDatabase();
+	~FAssetSearchDatabase();
+
+	FAssetSearchDatabase(const FAssetSearchDatabase&) = delete;
+	FAssetSearchDatabase& operator=(const FAssetSearchDatabase&) = delete;
+
+	FAssetSearchDatabase(FAssetSearchDatabase&&) = default;
+	FAssetSearchDatabase& operator=(FAssetSearchDatabase&&) = default;
+
+	/**
+	 * Is this a valid database? (ie, has been successfully opened).
+	 */
+	bool IsValid() const;
+
+	/**
+	 * Open (or create) a database file.
+	 *
+	 * @param InSessionPath				The root path to store all the data for this session under.
+	 *
+	 * @return True if the database was opened, false otherwise.
+	 */
+	bool Open(const FString& InSessionPath);
+
+	/**
+	 * Open (or create) a database file.
+	 *
+	 * @param InSessionPath				The root path to store all the data for this session under.
+	 * @param InOpenMode				How should the database be opened?
+	 *
+	 * @return True if the database was opened, false otherwise.
+	 */
+	bool Open(const FString& InSessionPath, const ESQLiteDatabaseOpenMode InOpenMode);
+
+	/**
+	 * Close an open database file.
+	 *
+	 * @param InDeleteDatabase			True if the session database and its associated data should be deleted after closing the database.
+	 */
+	bool Close(const bool InDeleteDatabase = false);
+
+	/**
+	 * Get the filename of the currently open database, or an empty string.
+	 * @note The returned filename will be an absolute pathname.
+	 */
+	FString GetFilename() const;
+
+	/**
+	 * Get the last error reported by this database.
+	 */
+	FString GetLastError() const;
+
+	/**
+	 * Is the asset cache up to date for this asset?
+	 */
+	bool IsAssetUpToDate(const FAssetData& InAssetData, const FString& IndexedJsonHash);
+
+	/**
+	 * Adds an asset to the search database.
+	 */
+	void AddOrUpdateAsset(const FAssetData& InAssetData, const FString& IndexedJson, const FString& IndexedJsonHash);
+
+	/**
+	 * Search for some stuff.
+	 */
+	bool EnumerateSearchResults(const FSearchQuery& Query, TFunctionRef<bool(FSearchRecord&&)> InCallback);
+
+private:
+	void LogLastError() const;
+
+private:
+	/** Internal SQLite database */
+	TUniquePtr<FSQLiteDatabase> Database;
+
+	/** The database filename name we use */
+	const FString DatabaseFileName;
+
+	/** Root path to store all session data under */
+	FString SessionPath;
+
+	/** Prepared statements for the currently open database */
+	TUniquePtr<FAssetSearchDatabaseStatements> Statements;
+};
\ No newline at end of file
diff --git a/Engine/Plugins/Editor/AssetSearch/Source/Private/AssetSearchManager.cpp b/Engine/Plugins/Editor/AssetSearch/Source/Private/AssetSearchManager.cpp
new file mode 100644
--- /dev/null
+++ b/Engine/Plugins/Editor/AssetSearch/Source/Private/AssetSearchManager.cpp
@@ -0,0 +1,353 @@
+// Copyright Epic Games, Inc. All Rights Reserved.
+
+#include "AssetSearchManager.h"
+#include "IAssetRegistry.h"
+#include "AssetRegistryModule.h"
+#include "AssetSearchDatabase.h"
+#include "Async/Async.h"
+#include "DerivedDataCacheInterface.h"
+
+#include "Indexers/DataAssetIndexer.h"
+#include "Indexers/DataTableIndexer.h"
+#include "Indexers/BlueprintIndexer.h"
+#include "Indexers/WidgetBlueprintIndexer.h"
+#include "Containers/StringConv.h"
+
+static bool bIndexUnindexAssetsOnLoad = false;
+FAutoConsoleVariableRef CVarIndexUnindexAssetsOnLoad(
+	TEXT("Search.IndexUnindexAssetsOnLoad"),
+	bIndexUnindexAssetsOnLoad,
+	TEXT("Index Unindex Assets On Load")
+);
+
+FAssetSearchManager::FAssetSearchManager()
+{
+	PendingDatabaseUpdates = 0;
+}
+
+FAssetSearchManager::~FAssetSearchManager()
+{
+	{
+		FScopeLock ScopedLock(&SearchDatabaseCS);
+		FCoreUObjectDelegates::OnObjectSaved.RemoveAll(this);
+		FCoreUObjectDelegates::OnAssetLoaded.RemoveAll(this);
+		UObject::FAssetRegistryTag::OnGetExtraObjectTags.RemoveAll(this);
+
+		//FModuleManager::GetModule<FAssetRegistryModule>("AssetRegistry");
+
+		FTicker::GetCoreTicker().RemoveTicker(TickerHandle);
+	}
+}
+
+void FAssetSearchManager::Start()
+{
+	RegisterIndexer(TEXT("DataAsset"), new FDataAssetIndexer());
+	RegisterIndexer(TEXT("DataTable"), new FDataTableIndexer());
+	RegisterIndexer(TEXT("Blueprint"), new FBlueprintIndexer());
+	RegisterIndexer(TEXT("WidgetBlueprint"), new FWidgetBlueprintIndexer());
+
+	const FString SessionPath = FPaths::ConvertRelativePathToFull(FPaths::Combine(FPaths::ProjectSavedDir(), TEXT("Search")));
+	SearchDatabase.Open(SessionPath);
+
+	FCoreUObjectDelegates::OnObjectSaved.AddRaw(this, &FAssetSearchManager::OnObjectSaved);
+	FCoreUObjectDelegates::OnAssetLoaded.AddRaw(this, &FAssetSearchManager::OnAssetLoaded);
+	UObject::FAssetRegistryTag::OnGetExtraObjectTags.AddRaw(this, &FAssetSearchManager::OnGetAssetTags);
+
+	IAssetRegistry& AssetRegistry = FModuleManager::LoadModuleChecked<FAssetRegistryModule>("AssetRegistry").Get();
+	AssetRegistry.OnAssetAdded().AddRaw(this, &FAssetSearchManager::OnAssetAdded);
+	AssetRegistry.GetAllAssets(ProcessAssetQueue, true);
+
+	TickerHandle = FTicker::GetCoreTicker().AddTicker(FTickerDelegate::CreateRaw(this, &FAssetSearchManager::Tick_GameThread), 0);
+}
+
+FSearchStats FAssetSearchManager::GetStats() const
+{
+	FSearchStats Stats;
+	Stats.Scanning = ProcessAssetQueue.Num();
+	Stats.Downloading = ProcessDDCQueue.Num();
+	Stats.PendingDatabaseUpdates = PendingDatabaseUpdates;
+	return Stats;
+}
+
+void FAssetSearchManager::RegisterIndexer(FName AssetClassName, IAssetIndexer* Indexer)
+{
+	check(IsInGameThread());
+	Indexers.Add(AssetClassName, Indexer);
+}
+
+void FAssetSearchManager::OnAssetAdded(const FAssetData& InAssetData)
+{
+	check(IsInGameThread());
+	ProcessAssetQueue.Add(InAssetData);
+}
+
+void FAssetSearchManager::OnObjectSaved(UObject* InObject)
+{
+	check(IsInGameThread());
+	StoreIndexForAsset(InObject);
+}
+
+void FAssetSearchManager::OnAssetLoaded(UObject* InObject)
+{
+	check(IsInGameThread());
+	if (bIndexUnindexAssetsOnLoad)
+	{
+		RequestIndexAsset(InObject);
+	}
+}
+
+void FAssetSearchManager::OnGetAssetTags(const UObject* Object, TArray<UObject::FAssetRegistryTag>& OutTags)
+{
+	check(IsInGameThread());
+
+	FString ObjectPath = Object->GetPathName();
+	const FContentHashEntry* ContentEntry = ContentHashCache.Find(ObjectPath);
+
+	if (ContentEntry)
+	{
+		static FName SearchIndexContentTag(TEXT("SearchIndexContent"));
+		static FName SearchIndexContentKeyTag(TEXT("SearchIndexContentKey"));
+		OutTags.Add(UObject::FAssetRegistryTag(SearchIndexContentKeyTag, *ContentEntry->ContentHash, UObject::FAssetRegistryTag::TT_Hidden));
+		OutTags.Add(UObject::FAssetRegistryTag(SearchIndexContentTag, *ContentEntry->Content, UObject::FAssetRegistryTag::TT_Hidden));
+	}
+}
+
+void FAssetSearchManager::AddToContentTagCache(const FAssetData& InAsset, const FString& InContent, const FString& InContentHash)
+{
+	FContentHashEntry Entry;
+	Entry.Content = InContent;
+	Entry.ContentHash = InContentHash;
+	ContentHashCache.Add(InAsset.ObjectPath.ToString(), Entry);
+}
+
+void FAssetSearchManager::RequestIndexAsset(UObject* InAsset)
+{
+	TWeakObjectPtr<UObject> AssetWeakPtr = InAsset;
+
+	FAssetData AssetData(InAsset);
+	FString AssetJsonDDCKey = TryGetDDCKeyForAsset(InAsset);
+	AsyncTask(ENamedThreads::AnyBackgroundThreadNormalTask, [this, AssetData, AssetWeakPtr, AssetJsonDDCKey]() {
+		FScopeLock ScopedLock(&SearchDatabaseCS);
+		if (!SearchDatabase.IsAssetUpToDate(AssetData, AssetJsonDDCKey))
+		{
+			AsyncTask(ENamedThreads::GameThread, [this, AssetWeakPtr]() {
+				StoreIndexForAsset(AssetWeakPtr.Get(), true);
+			});
+		}
+	});
+}
+
+bool FAssetSearchManager::IsAssetIndexable(UObject* InAsset)
+{
+	if (InAsset && InAsset->IsAsset())
+	{
+		// If it's not a permanent package, and one we just loaded for diffing, don't index it.
+		UPackage* Package = InAsset->GetOutermost();
+		if (Package->HasAnyPackageFlags(LOAD_ForDiff | LOAD_PackageForPIE | LOAD_ForFileDiff))
+		{
+			return false;
+		}
+
+		if (InAsset->HasAnyFlags(RF_Transient))
+		{
+			return false;
+		}
+
+		return true;
+	}
+
+	return false;
+}
+
+FString FAssetSearchManager::TryGetDDCKeyForAsset(const FAssetData& InAsset)
+{
+	FString AssetJsonDDCKey;
+
+	{
+		FString ContentDDCKey;
+		static FName SearchIndexContentKeyTag(TEXT("SearchIndexContentKey"));
+		if (InAsset.GetTagValue(SearchIndexContentKeyTag, ContentDDCKey))
+		{
+			AssetJsonDDCKey = ContentDDCKey;
+		}
+		else
+		{
+			UClass* AssetClass = InAsset.GetClass();
+			if (HasIndexerForClass(AssetClass))
+			{
+				const FString UnindexedAssetKey = GetDerivedDataKey(InAsset);
+				AssetJsonDDCKey = UnindexedAssetKey;
+			}
+		}
+	}
+
+	return AssetJsonDDCKey;
+}
+
+void FAssetSearchManager::TryLoadIndexForAsset(const FAssetData& InAsset)
+{
+	check(IsInGameThread());
+	FString AssetJsonDDCKey = TryGetDDCKeyForAsset(InAsset);
+
+	if (!AssetJsonDDCKey.IsEmpty())
+	{
+		AsyncTask(ENamedThreads::AnyBackgroundThreadNormalTask, [this, InAsset, AssetJsonDDCKey]() {
+			FScopeLock ScopedLock(&SearchDatabaseCS);
+			if (!SearchDatabase.IsAssetUpToDate(InAsset, AssetJsonDDCKey))
+			{
+				FAssetDDCRequest DDCRequest;
+				DDCRequest.AssetData = InAsset;
+				DDCRequest.DDCKey_IndexDataHash = AssetJsonDDCKey;
+				DDCRequest.DDCHandle = GetDerivedDataCacheRef().GetAsynchronous(*AssetJsonDDCKey, InAsset.ObjectPath.ToString());
+				ProcessDDCQueue.Add(DDCRequest);
+			}
+		});
+	}
+}
+
+FString FAssetSearchManager::GetDerivedDataKey(const FSHAHash& IndexedContentHash)
+{
+	const FString DDCKey = TEXT("AssetSearch_B") + LexToString(FSearchSerializer::GetVersion()) + TEXT("_") + IndexedContentHash.ToString();
+	return DDCKey;
+}
+
+FString FAssetSearchManager::GetDerivedDataKey(const FAssetData& UnindexedAsset)
+{
+	FString ContentPath = UnindexedAsset.ObjectPath.ToString();
+
+	FSHAHash UnindexedAssetHash;
+	FSHA1::HashBuffer(*ContentPath, ContentPath.Len() * sizeof(FString::ElementType), UnindexedAssetHash.Hash);
+
+	const FString DDCKey = TEXT("AssetSearch_Legacy_B") + LexToString(FSearchSerializer::GetVersion()) + TEXT("_") + UnindexedAssetHash.ToString();
+	return DDCKey;
+}
+
+bool FAssetSearchManager::HasIndexerForClass(UClass* AssetClass)
+{
+	UClass* IndexableClass = AssetClass;
+	while (IndexableClass)
+	{
+		if (IAssetIndexer* Indexer = Indexers.FindRef(IndexableClass->GetFName()))
+		{
+			return true;
+		}
+
+		IndexableClass = IndexableClass->GetSuperClass();
+	}
+
+	return false;
+}
+
+void FAssetSearchManager::StoreIndexForAsset(UObject* InAsset, bool bLegacyIndexing)
+{
+	check(IsInGameThread());
+
+	if (IsAssetIndexable(InAsset))
+	{
+		FAssetData InAssetData(InAsset);
+
+		bool bWasIndexed = false;
+		FString IndexedJson;
+		{
+			FSearchSerializer Serializer(InAssetData, &IndexedJson);
+
+			UClass* IndexableClass = InAsset->GetClass();
+			while (IndexableClass)
+			{
+				if (IAssetIndexer* Indexer = Indexers.FindRef(IndexableClass->GetFName()))
+				{
+					bWasIndexed = true;
+					Serializer.BeginIndexer(Indexer);
+					Indexer->IndexAsset(InAsset, Serializer);
+					Serializer.EndIndexer();
+				}
+
+				IndexableClass = IndexableClass->GetSuperClass();
+			}
+		}
+
+		if (bWasIndexed)
+		{
+			// Hash the content so that we can store it in the DDC.
+			FSHAHash IndexedJsonHash;
+			FSHA1::HashBuffer(*IndexedJson, IndexedJson.Len() * sizeof(FString::ElementType), IndexedJsonHash.Hash);
+
+			const FString DerivedDataKey = bLegacyIndexing ? GetDerivedDataKey(InAssetData) : GetDerivedDataKey(IndexedJsonHash);
+
+			FTCHARToUTF8 IndexedJsonUTF8(*IndexedJson);
+			TArrayView<const uint8> IndexedJsonUTF8View((const uint8*)IndexedJsonUTF8.Get(), IndexedJsonUTF8.Length() * sizeof(UTF8CHAR));
+			GetDerivedDataCacheRef().Put(*DerivedDataKey, IndexedJsonUTF8View, InAssetData.ObjectPath.ToString(), bLegacyIndexing);
+
+			AddToContentTagCache(InAssetData, IndexedJson, DerivedDataKey);
+
+			AddOrUpdateAsset(InAssetData, IndexedJson, DerivedDataKey);
+		}
+	}
+}
+
+void FAssetSearchManager::LoadDDCContentIntoDatabase(const FAssetData& InAsset, const TArray<uint8>& Content, const FString& DerivedDataKey)
+{
+	FUTF8ToTCHAR WByteBuffer((const ANSICHAR*)Content.GetData(), Content.Num());
+	FString IndexedJson(WByteBuffer.Length(), WByteBuffer.Get());
+
+	AddOrUpdateAsset(InAsset, IndexedJson, DerivedDataKey);
+}
+
+void FAssetSearchManager::AddOrUpdateAsset(const FAssetData& InAssetData, const FString& IndexedJson, const FString& DerivedDataKey)
+{
+	check(IsInGameThread());
+
+	PendingDatabaseUpdates++;
+	AsyncTask(ENamedThreads::AnyBackgroundThreadNormalTask, [this, InAssetData, IndexedJson, DerivedDataKey]() {
+		FScopeLock ScopedLock(&SearchDatabaseCS);
+		SearchDatabase.AddOrUpdateAsset(InAssetData, IndexedJson, DerivedDataKey);
+		PendingDatabaseUpdates--;
+	});
+}
+
+bool FAssetSearchManager::Tick_GameThread(float DeltaTime)
+{
+	check(IsInGameThread());
+
+	int32 ScanLimit = 1000;
+	while (ProcessAssetQueue.Num() > 0 && ScanLimit > 0)
+	{
+		FAssetData Asset = ProcessAssetQueue.Pop(false);
+		TryLoadIndexForAsset(Asset);
+		ScanLimit--;
+	}
+
+	int32 DownloadLimit = 100;
+	while (ProcessDDCQueue.Num() > 0 && DownloadLimit > 0)
+	{
+		FScopeLock ScopedLock(&SearchDatabaseCS);
+
+		const FAssetDDCRequest& PendingRequest = ProcessDDCQueue[0];
+		if (GetDerivedDataCacheRef().PollAsynchronousCompletion(PendingRequest.DDCHandle))
+		{
+			bool bDataWasBuilt;
+
+			TArray<uint8> OutContent;
+			bool bGetSuccessful = GetDerivedDataCacheRef().GetAsynchronousResults(PendingRequest.DDCHandle, OutContent, &bDataWasBuilt);
+			if (bGetSuccessful)
+			{
+				LoadDDCContentIntoDatabase(PendingRequest.AssetData, OutContent, PendingRequest.DDCKey_IndexDataHash);
+			}
+
+			ProcessDDCQueue.RemoveAt(0, 1, false);
+			DownloadLimit--;
+			continue;
+		}
+		break;
+	}
+
+	return true;
+}
+
+bool FAssetSearchManager::Search(const FSearchQuery& Query, TFunctionRef<bool(FSearchRecord&&)> InCallback)
+{
+	check(IsInGameThread());
+
+	FScopeLock ScopedLock(&SearchDatabaseCS);
+	return SearchDatabase.EnumerateSearchResults(Query, InCallback);
+}
\ No newline at end of file
diff --git a/Engine/Plugins/Editor/AssetSearch/Source/Private/AssetSearchManager.h b/Engine/Plugins/Editor/AssetSearch/Source/Private/AssetSearchManager.h
new file mode 100644
--- /dev/null
+++ b/Engine/Plugins/Editor/AssetSearch/Source/Private/AssetSearchManager.h
@@ -0,0 +1,74 @@
+// Copyright Epic Games, Inc. All Rights Reserved.
+
+#pragma once
+
+#include "CoreMinimal.h"
+
+class FRunnableThread;
+
+class FAssetSearchManager
+{
+public:
+	FAssetSearchManager();
+	~FAssetSearchManager();
+	
+	void Start();
+	void RegisterIndexer(FName AssetClassName, IAssetIndexer* Indexer);
+
+	FSearchStats GetStats() const;
+
+	bool Search(const FSearchQuery& Query, TFunctionRef<bool(FSearchRecord&&)> InCallback);
+
+private:
+	bool Tick_GameThread(float DeltaTime);
+
+private:
+	void OnAssetAdded(const FAssetData& InAssetData);
+	void OnObjectSaved(UObject* InObject);
+	void OnAssetLoaded(UObject* InObject);
+	void OnGetAssetTags(const UObject* Object, TArray<UObject::FAssetRegistryTag>& OutTags);
+	void AddToContentTagCache(const FAssetData& InAsset, const FString& InContent, const FString& InContentHash);
+
+	void AddOrUpdateAsset(const FAssetData& InAsset, const FString& IndexedJson, const FString& DerivedDataKey);
+
+	void RequestIndexAsset(UObject* InAsset);
+	bool IsAssetIndexable(UObject* InAsset);
+	void TryLoadIndexForAsset(const FAssetData& InAsset);
+	FString TryGetDDCKeyForAsset(const FAssetData& InAsset);
+	FString GetDerivedDataKey(const FSHAHash& IndexedContentHash);
+	FString GetDerivedDataKey(const FAssetData& UnindexedAsset);
+	bool HasIndexerForClass(UClass* AssetClass);
+	void StoreIndexForAsset(UObject* InAsset, bool Unindexed = false);
+	void LoadDDCContentIntoDatabase(const FAssetData& InAsset, const TArray<uint8>& Content, const FString& DerivedDataKey);
+
+private:
+	struct FContentHashEntry
+	{
+		FString ContentHash;
+		FString Content;
+	};
+	TMap<FString /*AssetPath*/, FContentHashEntry> ContentHashCache;
+
+private:
+	FAssetSearchDatabase SearchDatabase;
+	FCriticalSection SearchDatabaseCS;
+	TAtomic<int32> PendingDatabaseUpdates;
+
+private:
+	TMap<FName, IAssetIndexer*> Indexers;
+
+	TArray<FAssetData> ProcessAssetQueue;
+
+	struct FAssetDDCRequest
+	{
+		FAssetData AssetData;
+		FString DDCKey_IndexDataHash;
+		uint32 DDCHandle;
+	};
+	TArray<FAssetDDCRequest> ProcessDDCQueue;
+
+	FDelegateHandle TickerHandle;
+
+private:
+	FRunnableThread* Thread;
+};
\ No newline at end of file
diff --git a/Engine/Plugins/Editor/AssetSearch/Source/Private/AssetSearchModule.cpp b/Engine/Plugins/Editor/AssetSearch/Source/Private/AssetSearchModule.cpp
new file mode 100644
--- /dev/null
+++ b/Engine/Plugins/Editor/AssetSearch/Source/Private/AssetSearchModule.cpp
@@ -0,0 +1,91 @@
+// Copyright Epic Games, Inc. All Rights Reserved.
+
+#include "CoreMinimal.h"
+#include "Modules/ModuleManager.h"
+#include "IAssetSearchModule.h"
+#include "Widgets/DeclarativeSyntaxSupport.h"
+#include "WorkspaceMenuStructure.h"
+#include "EditorStyleSet.h"
+#include "Editor/WorkspaceMenuStructure/Public/WorkspaceMenuStructureModule.h"
+#include "Widgets/Docking/SDockTab.h"
+#include "Widgets/SSearchBrowser.h"
+#include "AssetRegistryModule.h"
+
+#define LOCTEXT_NAMESPACE "FAssetSearchModule"
+
+PRAGMA_DISABLE_OPTIMIZATION
+
+static const FName SearchTabName("Search");
+
+class FAssetSearchModule : public IAssetSearchModule
+{
+public:
+
+	// IAssetSearchModule interface
+
+
+public:
+
+	// IModuleInterface interface
+	
+	virtual void StartupModule() override
+	{
+		SearchManager = MakeUnique<FAssetSearchManager>();
+		SearchManager->Start();
+
+		FGlobalTabmanager::Get()->RegisterNomadTabSpawner(SearchTabName, FOnSpawnTab::CreateRaw(this, &FAssetSearchModule::HandleSpawnSettingsTab))
+			.SetDisplayName(LOCTEXT("Search", "Search"))
+			.SetTooltipText(LOCTEXT("Search", "Search Tab"))
+			.SetIcon(FSlateIcon(FEditorStyle::GetStyleSetName(), "Symbols.SearchGlass"))
+			.SetGroup(WorkspaceMenu::GetMenuStructure().GetToolsCategory());
+	}
+
+	virtual void ShutdownModule() override
+	{
+		FGlobalTabmanager::Get()->UnregisterNomadTabSpawner(SearchTabName);
+	}
+
+	void ExecuteOpenObjectBrowser()
+	{
+		FGlobalTabmanager::Get()->TryInvokeTab(SearchTabName);
+	}
+
+	TSharedRef<SDockTab> HandleSpawnSettingsTab(const FSpawnTabArgs& SpawnTabArgs)
+	{
+		const TSharedRef<SDockTab> DockTab = SNew(SDockTab)
+			.TabRole(ETabRole::NomadTab);
+
+		DockTab->SetContent(SNew(SSearchBrowser));
+
+		return DockTab;
+	}
+
+	virtual FSearchStats GetStats() const override
+	{
+		return SearchManager->GetStats();
+	}
+
+	virtual bool Search(const FSearchQuery& Query, TFunctionRef<bool(FSearchRecord&&)> InCallback) override
+	{
+		return SearchManager->Search(Query, InCallback);
+	}
+
+	virtual void RegisterIndexer(FName AssetClassName, IAssetIndexer* Indexer) override
+	{
+		SearchManager->RegisterIndexer(AssetClassName, Indexer);
+	}
+
+	virtual void UnregisterIndexer(IAssetIndexer* Indexer) override
+	{
+		
+	}
+
+private:
+	TUniquePtr<FAssetSearchManager> SearchManager;
+};
+
+#undef LOCTEXT_NAMESPACE
+
+PRAGMA_ENABLE_OPTIMIZATION
+
+IMPLEMENT_MODULE(FAssetSearchModule, AssetSearch);
diff --git a/Engine/Plugins/Editor/AssetSearch/Source/Private/Indexers/BlueprintIndexer.cpp b/Engine/Plugins/Editor/AssetSearch/Source/Private/Indexers/BlueprintIndexer.cpp
new file mode 100644
--- /dev/null
+++ b/Engine/Plugins/Editor/AssetSearch/Source/Private/Indexers/BlueprintIndexer.cpp
@@ -0,0 +1,103 @@
+// Copyright Epic Games, Inc. All Rights Reserved.
+
+#include "BlueprintIndexer.h"
+#include "Engine/Blueprint.h"
+#include "EdGraph/EdGraphPin.h"
+#include "K2Node_CallFunction.h"
+#include "K2Node_Variable.h"
+#include "Kismet2/BlueprintEditorUtils.h"
+#include "Utility/IndexerUtilities.h"
+
+#define LOCTEXT_NAMESPACE "FBlueprintIndexer"
+
+PRAGMA_DISABLE_OPTIMIZATION
+
+enum class EBlueprintIndexerVersion
+{
+	Empty = 0,
+	Initial = 1,
+
+	Current = Initial,
+};
+
+int32 FBlueprintIndexer::GetVersion() const
+{
+	return (int32)EBlueprintIndexerVersion::Current;
+}
+
+void FBlueprintIndexer::IndexAsset(const UObject* InAssetObject, FSearchSerializer& Serializer)
+{
+	const UBlueprint* BP = Cast<UBlueprint>(InAssetObject);
+	check(BP);
+
+	//UGameplayAbilityBlueprint
+
+	if (UClass* GeneratedClass = BP->GeneratedClass)
+	{
+		if (UObject* CDO = GeneratedClass->GetDefaultObject())
+		{
+			Serializer.BeginIndexingObject(CDO, TEXT("Class Defaults"));
+			FIndexerUtilities::IterateIndexableProperties(CDO, [&Serializer](const FProperty* Property, const FString& Value) {
+				Serializer.IndexProperty(Property, Value);
+			});
+			Serializer.EndIndexingObject();
+		}
+	}
+
+	TArray<UEdGraph*> AllGraphs;
+	BP->GetAllGraphs(AllGraphs);
+
+	for (UEdGraph* Graph : AllGraphs)
+	{	
+		for (UEdGraphNode* Node : Graph->Nodes)
+		{
+			FText NodeType = LOCTEXT("Node", "Node");
+			FText NodeText = Node->GetNodeTitle(ENodeTitleType::MenuTitle);
+			Serializer.BeginIndexingObject(Node, NodeText);
+			Serializer.IndexProperty(TEXT("Name"), NodeText);
+			
+			if (UK2Node_CallFunction* FunctionNode = Cast<UK2Node_CallFunction>(Node))
+			{
+				if (UFunction* TargetFunction = FunctionNode->GetTargetFunction())
+				{
+					Serializer.IndexProperty(TEXT("TargetFunction"), TargetFunction->GetName());
+				}
+			}
+			if (UK2Node_Variable* VariableNode = Cast<UK2Node_Variable>(Node))
+			{
+				const FMemberReference& VariableReference = VariableNode->VariableReference;
+
+				Serializer.IndexProperty(TEXT("MemberName"), VariableReference.GetMemberName().ToString());
+				if (VariableReference.GetMemberGuid().IsValid())
+				{
+					Serializer.IndexProperty(TEXT("MemberGuid"), VariableReference.GetMemberGuid().ToString(EGuidFormats::Digits));
+				}
+				if (UClass* MemberParentClass = VariableReference.GetMemberParentClass())
+				{
+					Serializer.IndexProperty(TEXT("MemberParent"), MemberParentClass->GetPathName());
+				}
+				//Serializer.WriteValue(TEXT("bSelfContext"), VariableReference.IsSelfContext());
+			}
+
+			if (Node->GetAllPins().Num())
+			{
+				for (UEdGraphPin* Pin : Node->GetAllPins())
+				{
+					FText PinText = Pin->GetDisplayName();
+					if (PinText.IsEmpty())
+					{
+						continue;
+					}
+
+					Serializer.IndexProperty(TEXT("Pin"), PinText);
+				}
+			}
+
+			Serializer.EndIndexingObject();
+		}
+	}
+}
+
+#undef LOCTEXT_NAMESPACE
+
+PRAGMA_ENABLE_OPTIMIZATION
\ No newline at end of file
diff --git a/Engine/Plugins/Editor/AssetSearch/Source/Private/Indexers/BlueprintIndexer.h b/Engine/Plugins/Editor/AssetSearch/Source/Private/Indexers/BlueprintIndexer.h
new file mode 100644
--- /dev/null
+++ b/Engine/Plugins/Editor/AssetSearch/Source/Private/Indexers/BlueprintIndexer.h
@@ -0,0 +1,13 @@
+// Copyright Epic Games, Inc. All Rights Reserved.
+
+#pragma once
+
+#include "CoreMinimal.h"
+#include "IAssetSearchModule.h"
+
+class FBlueprintIndexer : public IAssetIndexer
+{
+	virtual FString GetName() const override { return TEXT("Blueprint"); }
+	virtual int32 GetVersion() const override;
+	virtual void IndexAsset(const UObject* InAssetObject, FSearchSerializer& Serializer) override;
+};
\ No newline at end of file
diff --git a/Engine/Plugins/Editor/AssetSearch/Source/Private/Indexers/DataAssetIndexer.cpp b/Engine/Plugins/Editor/AssetSearch/Source/Private/Indexers/DataAssetIndexer.cpp
new file mode 100644
--- /dev/null
+++ b/Engine/Plugins/Editor/AssetSearch/Source/Private/Indexers/DataAssetIndexer.cpp
@@ -0,0 +1,29 @@
+// Copyright Epic Games, Inc. All Rights Reserved.
+
+#include "DataAssetIndexer.h"
+#include "Utility/IndexerUtilities.h"
+#include "Engine/DataAsset.h"
+
+enum class EDataAssetIndexerVersion
+{
+	Empty = 0,
+	Initial = 1,
+
+	Current = Initial,
+};
+
+int32 FDataAssetIndexer::GetVersion() const
+{
+	return (int32)EDataAssetIndexerVersion::Current;
+}
+
+void FDataAssetIndexer::IndexAsset(const UObject* InAssetObject, FSearchSerializer& Serializer)
+{
+	Serializer.BeginIndexingObject(InAssetObject, TEXT("$self"));
+
+	FIndexerUtilities::IterateIndexableProperties(InAssetObject, [&Serializer](const FProperty* Property, const FString& Value) {
+		Serializer.IndexProperty(Property, Value);
+	});
+
+	Serializer.EndIndexingObject();
+}
\ No newline at end of file
diff --git a/Engine/Plugins/Editor/AssetSearch/Source/Private/Indexers/DataAssetIndexer.h b/Engine/Plugins/Editor/AssetSearch/Source/Private/Indexers/DataAssetIndexer.h
new file mode 100644
--- /dev/null
+++ b/Engine/Plugins/Editor/AssetSearch/Source/Private/Indexers/DataAssetIndexer.h
@@ -0,0 +1,13 @@
+// Copyright Epic Games, Inc. All Rights Reserved.
+
+#pragma once
+
+#include "CoreMinimal.h"
+#include "IAssetSearchModule.h"
+
+class FDataAssetIndexer : public IAssetIndexer
+{
+	virtual FString GetName() const override { return TEXT("DataAsset"); }
+	virtual int32 GetVersion() const override;
+	virtual void IndexAsset(const UObject* InAssetObject, FSearchSerializer& Serializer) override;
+};
\ No newline at end of file
diff --git a/Engine/Plugins/Editor/AssetSearch/Source/Private/Indexers/DataTableIndexer.cpp b/Engine/Plugins/Editor/AssetSearch/Source/Private/Indexers/DataTableIndexer.cpp
new file mode 100644
--- /dev/null
+++ b/Engine/Plugins/Editor/AssetSearch/Source/Private/Indexers/DataTableIndexer.cpp
@@ -0,0 +1,38 @@
+// Copyright Epic Games, Inc. All Rights Reserved.
+
+#include "DataTableIndexer.h"
+#include "Utility/IndexerUtilities.h"
+#include "Engine/DataTable.h"
+
+enum class EDataTableIndexerVersion
+{
+	Empty = 0,
+	Initial = 1,
+
+	Current = Initial,
+};
+
+int32 FDataTableIndexer::GetVersion() const
+{
+	return (int32)EDataTableIndexerVersion::Current;
+}
+
+void FDataTableIndexer::IndexAsset(const UObject* InAssetObject, FSearchSerializer& Serializer)
+{
+	const UDataTable* DataTable = Cast<UDataTable>(InAssetObject);
+	check(DataTable);
+
+	const TMap<FName, uint8*>& Rows = DataTable->GetRowMap();
+
+	Serializer.BeginIndexingObject(DataTable, TEXT("$self"));
+	for (const auto& Entry : Rows)
+	{
+		const FName& RowName = Entry.Key;
+		uint8* Row = Entry.Value;
+
+		FIndexerUtilities::IterateIndexableProperties(DataTable->GetRowStruct(), Row, [&Serializer](const FProperty* Property, const FString& Value) {
+			Serializer.IndexProperty(Property, Value);
+		});
+	}
+	Serializer.EndIndexingObject();
+}
\ No newline at end of file
diff --git a/Engine/Plugins/Editor/AssetSearch/Source/Private/Indexers/DataTableIndexer.h b/Engine/Plugins/Editor/AssetSearch/Source/Private/Indexers/DataTableIndexer.h
new file mode 100644
--- /dev/null
+++ b/Engine/Plugins/Editor/AssetSearch/Source/Private/Indexers/DataTableIndexer.h
@@ -0,0 +1,13 @@
+// Copyright Epic Games, Inc. All Rights Reserved.
+
+#pragma once
+
+#include "CoreMinimal.h"
+#include "IAssetSearchModule.h"
+
+class FDataTableIndexer : public IAssetIndexer
+{
+	virtual FString GetName() const override { return TEXT("DataTable"); }
+	virtual int32 GetVersion() const override;
+	virtual void IndexAsset(const UObject* InAssetObject, FSearchSerializer& Serializer) override;
+};
\ No newline at end of file
diff --git a/Engine/Plugins/Editor/AssetSearch/Source/Private/Indexers/WidgetBlueprintIndexer.cpp b/Engine/Plugins/Editor/AssetSearch/Source/Private/Indexers/WidgetBlueprintIndexer.cpp
new file mode 100644
--- /dev/null
+++ b/Engine/Plugins/Editor/AssetSearch/Source/Private/Indexers/WidgetBlueprintIndexer.cpp
@@ -0,0 +1,43 @@
+// Copyright Epic Games, Inc. All Rights Reserved.
+
+#include "WidgetBlueprintIndexer.h"
+#include "Engine/Blueprint.h"
+#include "WidgetBlueprint.h"
+#include "Components/Widget.h"
+
+#define LOCTEXT_NAMESPACE "FWidgetBlueprintIndexer"
+
+enum class EWidgetBlueprintIndexerVersion
+{
+	Empty = 0,
+	Initial = 1,
+
+	Current = Initial,
+};
+
+int32 FWidgetBlueprintIndexer::GetVersion() const
+{
+	return (int32)EWidgetBlueprintIndexerVersion::Current;
+}
+
+void FWidgetBlueprintIndexer::IndexAsset(const UObject* InAssetObject, FSearchSerializer& Serializer)
+{
+	const UWidgetBlueprint* BP = Cast<UWidgetBlueprint>(InAssetObject);
+	check(BP);
+
+	TArray<const UWidget*> AllWidgets = BP->GetAllSourceWidgets();
+
+	for (const UWidget* Widget : AllWidgets)
+	{
+		Serializer.BeginIndexingObject(Widget, Widget->GetDisplayLabel());
+		Serializer.IndexProperty(TEXT("Name"), Widget->GetDisplayLabel());
+
+		FIndexerUtilities::IterateIndexableProperties(Widget, [&Serializer](const FProperty* Property, const FString& Value) {
+			Serializer.IndexProperty(Property, Value);
+		});
+
+		Serializer.EndIndexingObject();
+	}
+}
+
+#undef LOCTEXT_NAMESPACE
\ No newline at end of file
diff --git a/Engine/Plugins/Editor/AssetSearch/Source/Private/Indexers/WidgetBlueprintIndexer.h b/Engine/Plugins/Editor/AssetSearch/Source/Private/Indexers/WidgetBlueprintIndexer.h
new file mode 100644
--- /dev/null
+++ b/Engine/Plugins/Editor/AssetSearch/Source/Private/Indexers/WidgetBlueprintIndexer.h
@@ -0,0 +1,13 @@
+// Copyright Epic Games, Inc. All Rights Reserved.
+
+#pragma once
+
+#include "CoreMinimal.h"
+#include "IAssetSearchModule.h"
+
+class FWidgetBlueprintIndexer : public IAssetIndexer
+{
+	virtual FString GetName() const override { return TEXT("WidgetBlueprint"); }
+	virtual int32 GetVersion() const override;
+	virtual void IndexAsset(const UObject* InAssetObject, FSearchSerializer& Serializer) override;
+};
\ No newline at end of file
diff --git a/Engine/Plugins/Editor/AssetSearch/Source/Private/SearchSerializer.cpp b/Engine/Plugins/Editor/AssetSearch/Source/Private/SearchSerializer.cpp
new file mode 100644
--- /dev/null
+++ b/Engine/Plugins/Editor/AssetSearch/Source/Private/SearchSerializer.cpp
@@ -0,0 +1,195 @@
+// Copyright Epic Games, Inc. All Rights Reserved.
+
+#include "SearchSerializer.h"
+#include "JsonObjectConverter.h"
+#include "Serialization/JsonWriter.h"
+#include "Serialization/JsonSerializer.h"
+#include "Policies/CondensedJsonPrintPolicy.h"
+
+#define LOCTEXT_NAMESPACE "FSearchSerializer"
+
+enum class ESearchIndexVersion
+{
+	Empty = 0,
+	Initial = 1,
+
+	Current = Initial,
+};
+
+static FName StaticGetNativeClassName(const UClass* InClass)
+{
+	while (InClass && !InClass->HasAnyClassFlags(CLASS_Native))
+	{
+		InClass = InClass->GetSuperClass();
+	}
+
+	return InClass ? InClass->GetFName() : NAME_None;
+}
+
+FSearchSerializer::FSearchSerializer(const FAssetData& InAsset, FArchive& InAr)
+	: JsonWriter(TJsonWriterFactory<TCHAR, TCondensedJsonPrintPolicy<TCHAR> >::Create(&InAr))
+{
+	Initialize(InAsset);
+}
+
+FSearchSerializer::FSearchSerializer(const FAssetData& InAsset, FString* const Stream)
+	: JsonWriter(TJsonWriterFactory<TCHAR, TCondensedJsonPrintPolicy<TCHAR> >::Create(Stream))
+{
+	Initialize(InAsset);
+}
+
+void FSearchSerializer::Initialize(const FAssetData& InAsset)
+{
+	AssetData = InAsset;
+
+	JsonWriter->WriteObjectStart();
+
+	JsonWriter->WriteValue(TEXT("version"), GetVersion());
+
+	//JsonWriter->WriteValue(TEXT("AssetName"), InAsset.AssetName.ToString());
+	//JsonWriter->WriteValue(TEXT("AssetPath"), InAsset.ObjectPath.ToString());
+	//JsonWriter->WriteValue(TEXT("AssetClass"), InAsset.AssetClass.ToString());
+
+	JsonWriter->WriteObjectStart(TEXT("indexers"));
+}
+
+FSearchSerializer::~FSearchSerializer()
+{
+	JsonWriter->WriteObjectEnd();
+	JsonWriter->WriteObjectEnd();
+	JsonWriter->Close();
+}
+
+int32 FSearchSerializer::GetVersion()
+{
+	return (int32)ESearchIndexVersion::Current;
+}
+
+void FSearchSerializer::BeginIndexer(IAssetIndexer* InIndexer)
+{
+	JsonWriter->WriteObjectStart(InIndexer->GetName());
+	JsonWriter->WriteValue(TEXT("version"), InIndexer->GetVersion());
+
+	JsonWriter->WriteArrayStart(TEXT("objects"));
+}
+
+void FSearchSerializer::EndIndexer()
+{
+	JsonWriter->WriteArrayEnd();
+	JsonWriter->WriteObjectEnd();
+}
+
+void FSearchSerializer::BeginIndexingObject(const UObject* InObject, const FText& InFriendlyName)
+{
+	BeginIndexingObject(InObject, *FTextInspector::GetSourceString(InFriendlyName));
+}
+
+void FSearchSerializer::BeginIndexingObject(const UObject* InObject, const FString& InFriendlyName)
+{
+	ensure(CurrentObject == nullptr);
+	ensure(InObject->GetOutermost()->GetFName() == AssetData.PackageName);
+
+	CurrentObject = InObject;
+	CurrentObjectName = InFriendlyName;
+	CurrentObjectClass = StaticGetNativeClassName(InObject->GetClass()).ToString();
+	CurrentObjectPath = InObject->GetPathName();
+}
+
+void FSearchSerializer::EndIndexingObject()
+{
+	if (Values.Num() > 0)
+	{
+		JsonWriter->WriteObjectStart();
+		JsonWriter->WriteValue(TEXT("name"), CurrentObjectName);
+		JsonWriter->WriteValue(TEXT("path"), CurrentObjectPath);
+		JsonWriter->WriteValue(TEXT("native_class"), CurrentObjectClass);
+
+		JsonWriter->WriteArrayStart(TEXT("properties"));
+
+		for (const FIndexedValue& IndexedValue : Values)
+		{
+			JsonWriter->WriteObjectStart();
+
+			JsonWriter->WriteValue(TEXT("name"), IndexedValue.PropertyName);
+			JsonWriter->WriteValue(TEXT("field"), IndexedValue.PropertyFieldClass);
+
+			if (IndexedValue.PropertyClass != NAME_None)
+			{
+				JsonWriter->WriteValue(TEXT("class"), IndexedValue.PropertyClass.ToString());
+			}
+
+			if (!IndexedValue.Text.IsEmpty())
+			{
+				JsonWriter->WriteValue(TEXT("value_text"), IndexedValue.Text);
+			}
+
+			if (!IndexedValue.HiddenText.IsEmpty())
+			{
+				JsonWriter->WriteValue(TEXT("value_hidden"), IndexedValue.HiddenText);
+			}
+
+			JsonWriter->WriteObjectEnd();
+		}
+
+		JsonWriter->WriteArrayEnd();
+
+		JsonWriter->WriteObjectEnd();
+	}
+
+	Values.Reset();
+
+	CurrentObject = nullptr;
+	CurrentObjectClass.Reset();
+	CurrentObjectPath.Reset();
+}
+
+void FSearchSerializer::IndexProperty(const FString& InName, const FText& InValue)
+{
+	IndexProperty(*FTextProperty::StaticClass(), nullptr, InName, *FTextInspector::GetSourceString(InValue));
+}
+
+void FSearchSerializer::IndexProperty(const FString& InName, const FString& InValue)
+{
+	IndexProperty(*FStrProperty::StaticClass(), nullptr, InName, InValue);
+}
+
+void FSearchSerializer::IndexProperty(const UClass* InPropertyClass, const FString& InName, const FString& InValue)
+{
+	IndexProperty(*FObjectProperty::StaticClass(), InPropertyClass, InName, InValue);
+}
+
+void FSearchSerializer::IndexProperty(const UClass* InPropertyClass, const FString& InName, const FText& InValue)
+{
+	IndexProperty(*FObjectProperty::StaticClass(), InPropertyClass, InName, *FTextInspector::GetSourceString(InValue));
+}
+
+void FSearchSerializer::IndexProperty(const FProperty* InProperty, const FString& InValue)
+{
+	const FObjectProperty* ObjectProperty = CastField<FObjectProperty>(InProperty);
+	IndexProperty(*InProperty->GetClass(), ObjectProperty ? ObjectProperty->PropertyClass : nullptr, InProperty->GetDisplayNameText().ToString(), InValue);
+}
+
+void FSearchSerializer::IndexProperty(const FFieldClass& InPropertyFieldClass, const UClass* InPropertyClass, const FString& InName, const FString& InValue)
+{
+	if (ensure(CurrentObject))
+	{
+		// Don't index empty stuff.
+		if (InValue.IsEmpty())
+		{
+			return;
+		}
+
+		FIndexedValue& IndexedValue = Values.Emplace_GetRef();
+		IndexedValue.PropertyName = InName;
+		IndexedValue.PropertyFieldClass = InPropertyFieldClass.GetName();
+		IndexedValue.PropertyClass = StaticGetNativeClassName(InPropertyClass);
+		IndexedValue.Text = InValue;
+	}
+
+	///** Gets the list of all field classes in existence */
+	//static TArray<FFieldClass*>& GetAllFieldClasses();
+	///** Gets a mapping of all field class names to the actually class objects */
+	//static TMap<FName, FFieldClass*>& GetNameToFieldClassMap();
+}
+
+#undef LOCTEXT_NAMESPACE
diff --git a/Engine/Plugins/Editor/AssetSearch/Source/Private/Utility/IndexerUtilities.cpp b/Engine/Plugins/Editor/AssetSearch/Source/Private/Utility/IndexerUtilities.cpp
new file mode 100644
--- /dev/null
+++ b/Engine/Plugins/Editor/AssetSearch/Source/Private/Utility/IndexerUtilities.cpp
@@ -0,0 +1,117 @@
+// Copyright Epic Games, Inc. All Rights Reserved.
+
+#include "Utility/IndexerUtilities.h"
+#include "UObject/UnrealType.h"
+#include "GameplayTagContainer.h"
+#include "Internationalization/Text.h"
+
+PRAGMA_DISABLE_OPTIMIZATION
+
+static bool IsPropertyIndexable(const TPropertyValueIterator<FProperty>& It, const FProperty* Property)
+{
+	// Don't index transient properties.
+	if (Property->HasAnyPropertyFlags(CPF_Transient))
+	{
+		return false;
+	}
+
+	// Don't index anything we don't expose to the editor.
+	if (!Property->HasAnyPropertyFlags(CPF_Edit | CPF_BlueprintVisible | CPF_AssetRegistrySearchable))
+	{
+		TArray<const FProperty*> PropertyChain;
+		It.GetPropertyChain(PropertyChain);
+
+		if (PropertyChain.Num() > 0)
+		{
+			// If this is a property in a container property (that's not a struct) then that array is indexable, so this should be too.
+			if (CastField<FArrayProperty>(PropertyChain.Last()) || CastField<FSetProperty>(PropertyChain.Last()) || CastField<FMapProperty>(PropertyChain.Last()))
+			{
+				return true;
+			}
+		}
+		
+		return false;
+	}
+
+	return true;
+}
+
+void FIndexerUtilities::IterateIndexableProperties(const UObject* InObject, TFunctionRef<void(const FProperty* /*Property*/, const FString& /*Value*/)> Callback)
+{
+	IterateIndexableProperties(InObject->GetClass(), InObject, Callback);
+}
+
+void FIndexerUtilities::IterateIndexableProperties(const UStruct* InStruct, const void* InStructValue, TFunctionRef<void(const FProperty* /*Property*/, const FString& /*Value*/)> Callback)
+{
+	for (TPropertyValueIterator<FProperty> It(InStruct, InStructValue, EPropertyValueIteratorFlags::FullRecursion, EFieldIteratorFlags::ExcludeDeprecated); It; ++It)
+	{
+		const FProperty* Property = It.Key();
+
+		// Don't index a property we wouldn't normally expose to the user.
+		if (!IsPropertyIndexable(It, Property))
+		{
+			It.SkipRecursiveProperty();
+			continue;
+		}
+
+		void const* ValuePtr = It.Value();
+
+		FString Text;
+		if (const FNameProperty* NameProperty = CastField<FNameProperty>(Property))
+		{
+			const FName Value = NameProperty->GetPropertyValue(ValuePtr);
+			Text = Value.ToString();
+		}
+		else if (const FStrProperty* StringProperty = CastField<FStrProperty>(Property))
+		{
+			const FString Value = StringProperty->GetPropertyValue(ValuePtr);
+			Text = Value;
+		}
+		else if (const FTextProperty* TextProperty = CastField<FTextProperty>(Property))
+		{
+			const FText Value = TextProperty->GetPropertyValue(ValuePtr);
+			Text = *FTextInspector::GetSourceString(Value);
+		}
+		else if (const FEnumProperty* EnumProperty = CastField<FEnumProperty>(Property))
+		{
+			const int64 Value = EnumProperty->GetUnderlyingProperty()->GetSignedIntPropertyValue(ValuePtr);
+			FText DisplayName = EnumProperty->GetEnum()->GetDisplayNameTextByValue(Value);
+			Text = *FTextInspector::GetSourceString(DisplayName);
+		}
+		else if (const FStructProperty* StructProperty = CastField<FStructProperty>(Property))
+		{
+			if (StructProperty->Struct == FGameplayTag::StaticStruct())
+			{
+				const FGameplayTag* GameplayTag = static_cast<const FGameplayTag*>(ValuePtr);
+				Text = GameplayTag->GetTagName().ToString();
+				
+				It.SkipRecursiveProperty();
+			}
+			else if (StructProperty->Struct == FGameplayTagContainer::StaticStruct())
+			{
+				const FGameplayTagContainer* GameplayTagContainer = static_cast<const FGameplayTagContainer*>(ValuePtr);
+				for (auto TagIter = GameplayTagContainer->CreateConstIterator(); TagIter; ++TagIter)
+				{
+					Text = TagIter->GetTagName().ToString();
+					Callback(Property, Text);
+				}
+
+				It.SkipRecursiveProperty();
+				continue;
+			}
+			//else if (StructProperty->Struct == FGuid::StaticStruct())
+			//{
+			//	const FGuid* Guid = static_cast<const FGuid*>(ValuePtr);
+			//	Text = Guid->ToString();
+			//}
+		}
+
+		// Ignore empty records
+		if (!Text.IsEmpty())
+		{
+			Callback(Property, Text);
+		}
+	}
+}
+
+PRAGMA_ENABLE_OPTIMIZATION
\ No newline at end of file
diff --git a/Engine/Plugins/Editor/AssetSearch/Source/Private/Widgets/SSearchBrowser.cpp b/Engine/Plugins/Editor/AssetSearch/Source/Private/Widgets/SSearchBrowser.cpp
new file mode 100644
--- /dev/null
+++ b/Engine/Plugins/Editor/AssetSearch/Source/Private/Widgets/SSearchBrowser.cpp
@@ -0,0 +1,271 @@
+// Copyright Epic Games, Inc. All Rights Reserved.
+
+#include "SSearchBrowser.h"
+#include "SSearchTreeRow.h"
+#include "EditorFontGlyphs.h"
+
+#include "ClassViewerModule.h"
+#include "Kismet2/SClassPickerDialog.h"
+#include "Logging/LogMacros.h"
+#include "UObject/ReferenceChainSearch.h"
+#include "UObject/UObjectIterator.h"
+#include "Modules/ModuleManager.h"
+#include "Framework/MultiBox/MultiBoxBuilder.h"
+
+// Common classes for the picker
+#include "Components/Widget.h"
+#include "Blueprint/UserWidget.h"
+#include "GameFramework/Pawn.h"
+#include "GameFramework/PlayerController.h"
+#include "Engine/World.h"
+#include "GameFramework/Actor.h"
+#include "Dialogs/Dialogs.h"
+#include "Widgets/Layout/SBox.h"
+#include "Widgets/Input/SButton.h"
+#include "Widgets/Input/SComboButton.h"
+#include "Widgets/Input/SEditableTextBox.h"
+#include "Widgets/Input/SMultiLineEditableTextBox.h"
+#include "SearchModel.h"
+#include "EditorStyleSet.h"
+#include "IAssetSearchModule.h"
+#include "AssetRegistryModule.h"
+#include "Widgets/Views/STreeView.h"
+#include "Widgets/Layout/SSeparator.h"
+
+#define LOCTEXT_NAMESPACE "SObjectBrowser"
+
+DEFINE_LOG_CATEGORY_STATIC(LogObjectBrowser, Log, All)
+
+namespace AssetSearchConstants
+{
+	/** The size of the thumbnail pool */
+	const int32 ThumbnailPoolSize = 64;
+}
+
+void SSearchBrowser::Construct( const FArguments& InArgs )
+{
+	SortByColumn = SSearchTreeRow::NAME_ColumnName;
+	SortMode = EColumnSortMode::Ascending;
+
+	const bool bAreRealTimeThumbnailsAllowed = false;
+	ThumbnailPool = MakeShared<FAssetThumbnailPool>(AssetSearchConstants::ThumbnailPoolSize, bAreRealTimeThumbnailsAllowed);
+
+	AssetRegistry = &FModuleManager::LoadModuleChecked<FAssetRegistryModule>("AssetRegistry").Get();
+
+	ChildSlot
+	[
+		SNew(SBorder)
+		.Padding(FMargin(3))
+		.BorderImage(FEditorStyle::GetBrush("ToolPanel.GroupBorder"))
+		[
+			SNew(SVerticalBox)
+
+			+ SVerticalBox::Slot()
+			.AutoHeight()
+			[
+				SNew(SHorizontalBox)
+
+				+ SHorizontalBox::Slot()
+				.AutoWidth()
+				[
+					SNew(SButton)
+					.OnClicked(this, &SSearchBrowser::OnRefresh)
+					[
+						SNew(STextBlock)
+						.Font(FEditorStyle::Get().GetFontStyle("FontAwesome.11"))
+						.Text(FEditorFontGlyphs::Refresh)
+					]
+				]
+
+				+ SHorizontalBox::Slot()
+				.FillWidth(1.0f)
+				[
+					SNew(SEditableTextBox)
+					.HintText(LOCTEXT("SearchHint", "Search"))
+					.OnTextCommitted(this, &SSearchBrowser::OnSearchTextCommited)
+				]
+			]
+
+			+ SVerticalBox::Slot()
+			.FillHeight(1.0f)
+			[
+				SNew(SBorder)
+				.BorderImage(FEditorStyle::GetBrush("ToolPanel.GroupBorder"))
+				.Padding(FMargin(0.0f, 4.0f))
+				[
+					SAssignNew(SearchTreeView, STreeView< TSharedPtr<FSearchNode> >)
+					.ItemHeight(24.0f)
+					.TreeItemsSource(&SearchResults)
+					.SelectionMode(ESelectionMode::Single)
+					.OnGenerateRow(this, &SSearchBrowser::HandleListGenerateRow)
+					.OnGetChildren(this, &SSearchBrowser::GetChildrenForInfo)
+					.OnSelectionChanged(this, &SSearchBrowser::HandleListSelectionChanged)
+					.HeaderRow
+					(
+						SNew(SHeaderRow)
+
+						+ SHeaderRow::Column(SSearchTreeRow::NAME_ColumnName)
+						.DefaultLabel(LOCTEXT("ColumnName", "Name"))
+						.FillWidth(0.80f)
+
+						+ SHeaderRow::Column(SSearchTreeRow::NAME_ColumnType)
+						.ManualWidth(300)
+						.DefaultLabel(LOCTEXT("ColumnType", "Type"))
+						//.SortMode(this, &SSearchBrowser::GetColumnSortMode, SObjectBrowserTableRow::CategoryClass)
+						//.OnSort(this, &SSearchBrowser::OnColumnSortModeChanged)
+
+					//	+ SHeaderRow::Column(SObjectBrowserTableRow::CategoryWorld)
+					//	.FixedWidth(22)
+					//	.DefaultLabel(LOCTEXT("CategoryWorld", ""))
+					//	.SortMode(this, &SSearchBrowser::GetColumnSortMode, SObjectBrowserTableRow::CategoryWorld)
+					//	.OnSort(this, &SSearchBrowser::OnColumnSortModeChanged)
+
+					//	+ SHeaderRow::Column(SObjectBrowserTableRow::CategoryClass)
+					//	.ManualWidth(300)
+					//	.DefaultLabel(LOCTEXT("CategoryAsset", "Asset"))
+					//	.SortMode(this, &SSearchBrowser::GetColumnSortMode, SObjectBrowserTableRow::CategoryClass)
+					//	.OnSort(this, &SSearchBrowser::OnColumnSortModeChanged)
+
+					//	+ SHeaderRow::Column(SObjectBrowserTableRow::CategoryProperty)
+					//	.FillWidth(1.0f)
+					//	.DefaultLabel(LOCTEXT("CategoryProperty", "Text"))
+					//	.SortMode(this, &SSearchBrowser::GetColumnSortMode, SObjectBrowserTableRow::CategoryProperty)
+					//	.OnSort(this, &SSearchBrowser::OnColumnSortModeChanged)
+
+					//	//+ SHeaderRow::Column(SObjectBrowserTableRow::CategoryPropertyValue)
+					//	//.DefaultLabel(LOCTEXT("CategoryPropertyValue", "Value"))
+					//	//.SortMode(this, &SSearchBrowser::GetColumnSortMode, SObjectBrowserTableRow::CategoryPropertyValue)
+					//	//.OnSort(this, &SSearchBrowser::OnColumnSortModeChanged)
+					)
+				]
+			]
+
+			+ SVerticalBox::Slot()
+			.AutoHeight()
+			.Padding(0, 0, 0, 1)
+			[
+				SNew(SSeparator)
+			]
+
+			+ SVerticalBox::Slot()
+			.AutoHeight()
+			.Padding(0, 0, 0, 1)
+			[
+				SNew(SHorizontalBox)
+
+				// Asset Stats 
+				+ SHorizontalBox::Slot()
+				.FillWidth(1.f)
+				.VAlign(VAlign_Center)
+				.Padding(8, 0)
+				[
+					SNew(STextBlock)
+					.Text(this, &SSearchBrowser::GetStatusText)
+				]
+			]
+		]
+	];
+
+	RefreshList();
+}
+
+FText SSearchBrowser::GetStatusText() const
+{
+	IAssetSearchModule& SearchModule = IAssetSearchModule::Get();
+	FSearchStats SearchStats = SearchModule.GetStats();
+	return FText::Format(LOCTEXT("SearchStatusTextFmt", "Scanning {0}   Downloading {1}   Updating {2}"), SearchStats.Scanning, SearchStats.Downloading, SearchStats.PendingDatabaseUpdates);
+}
+
+FReply SSearchBrowser::OnRefresh()
+{
+	RefreshList();
+
+	return FReply::Handled();
+}
+
+EColumnSortMode::Type SSearchBrowser::GetColumnSortMode(const FName ColumnId) const
+{
+	if (SortByColumn == ColumnId)
+	{
+		return SortMode;
+	}
+
+	return EColumnSortMode::None;
+}
+
+void SSearchBrowser::OnColumnSortModeChanged(const EColumnSortPriority::Type SortPriority, const FName& ColumnId, const EColumnSortMode::Type InSortMode)
+{
+	SortByColumn = ColumnId;
+	SortMode = InSortMode;
+
+	RefreshList();
+}
+
+void SSearchBrowser::RefreshList()
+{
+	SearchResults.Reset();
+	SearchResultHierarchy.Reset();
+
+	SearchTreeView->RequestListRefresh();
+
+	if (!FilterText.IsEmpty())
+	{
+		FSearchQuery Query;
+		Query.Query = FilterText.ToString();
+
+		IAssetSearchModule& SearchModule = IAssetSearchModule::Get();
+		SearchModule.Search(Query, [this](FSearchRecord&& InResult) {
+			AppendResult(MoveTemp(InResult));
+			return true;
+		});
+	}
+}
+
+void SSearchBrowser::AppendResult(FSearchRecord&& InResult)
+{
+	TSharedPtr<FAssetNode> ExistingAssetNode = SearchResultHierarchy.FindRef(InResult.AssetPath);
+	if (!ExistingAssetNode.IsValid())
+	{
+		ExistingAssetNode = MakeShared<FAssetNode>(InResult);
+		SearchResultHierarchy.Add(InResult.AssetPath, ExistingAssetNode);
+	}
+	else
+	{
+		ExistingAssetNode->Append(InResult);
+	}
+
+	SearchResults.Reset();
+	for (auto& Entry : SearchResultHierarchy)
+	{
+		SearchResults.Add(Entry.Value);
+	}
+
+	SearchTreeView->RequestListRefresh();
+	SearchTreeView->SetItemExpansion(ExistingAssetNode, true);
+}
+
+void SSearchBrowser::OnSearchTextCommited(const FText& InText, ETextCommit::Type InCommitType)
+{
+	FilterText = InText;
+	FilterString = FilterText.ToString();
+
+	RefreshList();
+}
+
+TSharedRef<ITableRow> SSearchBrowser::HandleListGenerateRow(TSharedPtr<FSearchNode> ObjectPtr, const TSharedRef<STableViewBase>& OwnerTable)
+{
+	return SNew(SSearchTreeRow, OwnerTable, AssetRegistry, ThumbnailPool)
+		.Object(ObjectPtr)
+		.HighlightText(FilterText);
+}
+
+void SSearchBrowser::GetChildrenForInfo(TSharedPtr<FSearchNode> InNode, TArray< TSharedPtr<FSearchNode> >& OutChildren)
+{
+	InNode->GetChildren(OutChildren);
+}
+
+void SSearchBrowser::HandleListSelectionChanged(TSharedPtr<FSearchNode> InNode, ESelectInfo::Type SelectInfo)
+{
+}
+
+#undef LOCTEXT_NAMESPACE
diff --git a/Engine/Plugins/Editor/AssetSearch/Source/Private/Widgets/SSearchBrowser.h b/Engine/Plugins/Editor/AssetSearch/Source/Private/Widgets/SSearchBrowser.h
new file mode 100644
--- /dev/null
+++ b/Engine/Plugins/Editor/AssetSearch/Source/Private/Widgets/SSearchBrowser.h
@@ -0,0 +1,75 @@
+// Copyright Epic Games, Inc. All Rights Reserved.
+
+#pragma once
+
+#include "Widgets/SCompoundWidget.h"
+#include "Widgets/Views/SHeaderRow.h"
+#include "Widgets/Views/SListView.h"
+#include "Widgets/Views/STableViewBase.h"
+
+class IDetailsView;
+class ITableRow;
+class FSearchNode;
+class FAssetNode;
+class FMenuBuilder;
+class IAssetRegistry;
+
+/**
+ * Implements the undo history panel.
+ */
+class SSearchBrowser : public SCompoundWidget
+{
+public:
+
+	SLATE_BEGIN_ARGS(SSearchBrowser) { }
+	SLATE_END_ARGS()
+
+public:
+
+	/**
+	 * Construct this widget
+	 *
+	 * @param InArgs The declaration data for this widget.
+	 */
+	void Construct( const FArguments& InArgs );
+
+public:
+
+private:
+
+	FText GetStatusText() const;
+
+	FReply OnRefresh();
+
+	EColumnSortMode::Type GetColumnSortMode(const FName ColumnId) const;
+	void OnColumnSortModeChanged(const EColumnSortPriority::Type SortPriority, const FName& ColumnId, const EColumnSortMode::Type InSortMode);
+
+	void RefreshList();
+	void AppendResult(FSearchRecord&& InResult);
+
+	void OnSearchTextCommited(const FText& InText, ETextCommit::Type InCommitType);
+
+	TSharedRef<ITableRow> HandleListGenerateRow(TSharedPtr<FSearchNode> TransactionInfo, const TSharedRef<STableViewBase>& OwnerTable);
+	void GetChildrenForInfo(TSharedPtr<FSearchNode> InInfo, TArray< TSharedPtr<FSearchNode> >& OutChildren);
+
+	void HandleListSelectionChanged(TSharedPtr<FSearchNode> TransactionInfo, ESelectInfo::Type SelectInfo);
+
+private:
+
+	FText FilterText;
+
+	// Filters
+	FString FilterString;
+	
+	TMap<FString, TSharedPtr<FAssetNode>> SearchResultHierarchy;
+	TArray< TSharedPtr<FSearchNode> > SearchResults;
+
+	TSharedPtr< STreeView< TSharedPtr<FSearchNode> > > SearchTreeView;
+
+	IAssetRegistry* AssetRegistry;
+
+	TSharedPtr<FAssetThumbnailPool> ThumbnailPool;
+
+	FName SortByColumn;
+	EColumnSortMode::Type SortMode;
+};
diff --git a/Engine/Plugins/Editor/AssetSearch/Source/Private/Widgets/SSearchTreeRow.cpp b/Engine/Plugins/Editor/AssetSearch/Source/Private/Widgets/SSearchTreeRow.cpp
new file mode 100644
--- /dev/null
+++ b/Engine/Plugins/Editor/AssetSearch/Source/Private/Widgets/SSearchTreeRow.cpp
@@ -0,0 +1,183 @@
+// Copyright Epic Games, Inc. All Rights Reserved.
+
+#include "SSearchTreeRow.h"
+#include "SearchModel.h"
+#include "Widgets/Text/STextBlock.h"
+#include "Widgets/Images/SImage.h"
+#include "Editor.h"
+#include "Subsystems/AssetEditorSubsystem.h"
+#include "Styling/SlateIconFinder.h"
+#include "AssetThumbnail.h"
+#include "IAssetRegistry.h"
+#include "AssetToolsModule.h"
+#include "Kismet2/KismetEditorUtilities.h"
+
+#define LOCTEXT_NAMESPACE "SObjectBrowserTableRow"
+
+FName SSearchTreeRow::NAME_ColumnName(TEXT("Name"));
+FName SSearchTreeRow::NAME_ColumnType(TEXT("Type"));
+//FName SSearchTreeRow::CategoryProperty(TEXT("Property"));
+//FName SSearchTreeRow::CategoryPropertyValue(TEXT("PropertyValue"));
+
+void SSearchTreeRow::Construct(const FArguments& InArgs, const TSharedRef<STableViewBase>& InOwnerTableView, IAssetRegistry* InAssetRegistry, TSharedPtr<FAssetThumbnailPool> InThumbnailPool)
+{
+	BrowserObject = InArgs._Object;
+	AssetRegistry = InAssetRegistry;
+	ThumbnailPool = InThumbnailPool;
+
+	FSuperRowType::Construct(
+		FSuperRowType::FArguments(),
+		InOwnerTableView
+	);
+}
+
+TSharedRef<SWidget> SSearchTreeRow::GenerateWidget()
+{
+	return SNew(STextBlock)
+		.Text(FText::FromString(BrowserObject->GetText()));
+}
+
+TSharedRef<SWidget> SSearchTreeRow::GenerateWidgetForColumn(const FName& ColumnName)
+{
+	TSharedRef<SHorizontalBox> HorizBox = SNew(SHorizontalBox);
+
+	FAssetToolsModule& AssetToolsModule = FModuleManager::GetModuleChecked<FAssetToolsModule>("AssetTools");
+
+	if (ColumnName == NAME_ColumnName)
+	{
+		TSharedPtr<SWidget> IconWidget;
+		if (BrowserObject->GetType() == ESearchNodeType::Object)
+		{
+			TSharedPtr<FAssetObjectNode> ObjectNode = StaticCastSharedPtr<FAssetObjectNode>(BrowserObject);
+
+			if (UClass* ObjectClass = FindObject<UClass>(ANY_PACKAGE, *ObjectNode->object_native_class, true))
+			{
+				//TSharedPtr<IAssetTypeActions> AssetActions = AssetToolsModule.Get().GetAssetTypeActionsForClass(ObjectClass).Pin();
+
+				FSlateIcon ClassIcon = FSlateIconFinder::FindIconForClass(ObjectClass);
+				if (ClassIcon.IsSet())
+				{
+					IconWidget = SNew(SImage)
+						.Image(ClassIcon.GetIcon());
+						//.ColorAndOpacity(AssetActions.IsValid() ? AssetActions->GetTypeColor() : FColor::White);
+				}
+			}
+		}
+		else if (BrowserObject->GetType() == ESearchNodeType::Asset)
+		{
+			TSharedPtr<FAssetNode> ObjectNode = StaticCastSharedPtr<FAssetNode>(BrowserObject);
+
+			if (UClass* ObjectClass = FindObject<UClass>(ANY_PACKAGE, *ObjectNode->AssetClass, true))
+			{
+				TSharedPtr<IAssetTypeActions> AssetActions = AssetToolsModule.Get().GetAssetTypeActionsForClass(ObjectClass).Pin();
+
+				FSlateIcon ClassIcon = FSlateIconFinder::FindIconForClass(ObjectClass);
+				if (ClassIcon.IsSet())
+				{
+					IconWidget = SNew(SImage)
+						.Image(ClassIcon.GetIcon())
+						.ColorAndOpacity(AssetActions.IsValid() ? AssetActions->GetTypeColor() : FColor::White);
+				}
+			}
+		}
+
+		HorizBox->AddSlot()
+			.VAlign(VAlign_Center)
+			.AutoWidth()
+			[
+				SNew(SExpanderArrow, SharedThis(this))
+			];
+
+		if (IconWidget.IsValid())
+		{
+			HorizBox->AddSlot()
+				.VAlign(VAlign_Center)
+				.AutoWidth()
+				.Padding(0, 0, 4, 0)
+				[
+					IconWidget.ToSharedRef()
+				];
+		}
+
+		HorizBox->AddSlot()
+			.VAlign(VAlign_Center)
+			.AutoWidth()
+			[
+				GenerateWidget()
+			];
+	}
+	else if(ColumnName == NAME_ColumnType)
+	{
+		if (BrowserObject->GetType() == ESearchNodeType::Asset)
+		{
+			TSharedPtr<FAssetNode> AssetNode = StaticCastSharedPtr<FAssetNode>(BrowserObject);
+			
+			HorizBox->AddSlot()
+				.VAlign(VAlign_Center)
+				.AutoWidth()
+				[
+					SNew(STextBlock)
+					.Text(FText::FromString(AssetNode->AssetClass))
+				];
+		}
+	}
+
+		////FAssetData AssetData = AssetRegistry->GetAssetByObjectPath(*BrowserObject->GetAsset);
+
+		//// Create the thumbnail handle
+		//int32 ThumbnailSizeX = 24;
+		//int32 ThumbnailSizeY = 24;
+		//TSharedRef<FAssetThumbnail> AssetThumbnail = MakeShared<FAssetThumbnail>(AssetData, ThumbnailSizeX, ThumbnailSizeY, ThumbnailPool);
+
+		//return SNew(SBox)
+		//	.WidthOverride(ThumbnailSizeX)
+		//	.HeightOverride(ThumbnailSizeY)
+		//	[
+		//		AssetThumbnail->MakeThumbnailWidget()
+		//	];
+	
+	//else if (ColumnName == CategoryClass)
+	//{
+	//	return SNew(STextBlock)
+	//		.Text(FText::FromString(BrowserObject->GetText()));
+	//}
+	//else if (ColumnName == CategoryProperty)
+	//{
+	//	//return SNew(STextBlock)
+	//	//	.Text(FText::FromString(BrowserObject->GetText()));
+	//}
+	//else if (ColumnName == CategoryPropertyValue)
+	//{
+	//	//return SNew(STextBlock)
+	//	//	.Text(FText::FromString(BrowserObject->Record.Text));
+	//}
+
+	return HorizBox;
+}
+
+FReply SSearchTreeRow::OnMouseButtonDoubleClick(const FGeometry& InMyGeometry, const FPointerEvent& InMouseEvent)
+{
+	if (BrowserObject->GetType() == ESearchNodeType::Asset)
+	{
+		TSharedPtr<FAssetNode> ObjectNode = StaticCastSharedPtr<FAssetNode>(BrowserObject);
+
+		GEditor->GetEditorSubsystem<UAssetEditorSubsystem>()->OpenEditorForAsset(ObjectNode->AssetPath);
+	}
+	else if (BrowserObject->GetType() == ESearchNodeType::Object || BrowserObject->GetType() == ESearchNodeType::Property)
+	{
+		FSoftObjectPath ReferencePath(BrowserObject->GetObjectPath());
+		UObject* Object = ReferencePath.TryLoad();
+		if (Object->GetTypedOuter<UBlueprint>())
+		{
+			FKismetEditorUtilities::BringKismetToFocusAttentionOnObject(Object, false);
+		}
+		else
+		{
+			GEditor->GetEditorSubsystem<UAssetEditorSubsystem>()->OpenEditorForAsset(ReferencePath.GetAssetPathString());
+		}
+	}
+
+	return FReply::Handled();
+}
+
+#undef LOCTEXT_NAMESPACE
diff --git a/Engine/Plugins/Editor/AssetSearch/Source/Private/Widgets/SSearchTreeRow.h b/Engine/Plugins/Editor/AssetSearch/Source/Private/Widgets/SSearchTreeRow.h
new file mode 100644
--- /dev/null
+++ b/Engine/Plugins/Editor/AssetSearch/Source/Private/Widgets/SSearchTreeRow.h
@@ -0,0 +1,63 @@
+// Copyright Epic Games, Inc. All Rights Reserved.
+
+#pragma once
+
+#include "ClassIconFinder.h"
+#include "Toolkits/AssetEditorManager.h"
+#include "CoreMinimal.h"
+#include "Widgets/SWidget.h"
+#include "Widgets/Views/STableRow.h"
+#include "Widgets/Views/SListView.h"
+#include "Widgets/DeclarativeSyntaxSupport.h"
+#include "Input/Reply.h"
+#include "Widgets/Views/STableViewBase.h"
+
+#define LOCTEXT_NAMESPACE "SObjectBrowserTableRow"
+
+class FSearchNode;
+
+/**
+ * Implements a row widget for a live uobject.
+ */
+class SSearchTreeRow : public SMultiColumnTableRow<TSharedPtr<FSearchNode>>
+{
+public:
+
+	SLATE_BEGIN_ARGS(SSearchTreeRow) { }
+		SLATE_ARGUMENT(TSharedPtr<FSearchNode>, Object)
+		SLATE_ARGUMENT(FText, HighlightText)
+	SLATE_END_ARGS()
+
+public:
+	static FName NAME_ColumnName;
+	static FName NAME_ColumnType;
+	//static FName CategoryProperty;
+	//static FName CategoryPropertyValue;
+
+public:
+
+	SSearchTreeRow()
+	{
+	}
+
+	void Construct(const FArguments& InArgs, const TSharedRef<STableViewBase>& InOwnerTableView, IAssetRegistry* InAssetRegistry, TSharedPtr<FAssetThumbnailPool> InThumbnailPool);
+
+public:
+
+	virtual TSharedRef<SWidget> GenerateWidget();
+
+	virtual TSharedRef<SWidget> GenerateWidgetForColumn(const FName& ColumnName) override;
+	
+	virtual FReply OnMouseButtonDoubleClick(const FGeometry& InMyGeometry, const FPointerEvent& InMouseEvent) override;
+
+private:
+
+	TSharedPtr<FSearchNode> BrowserObject;
+
+	IAssetRegistry* AssetRegistry;
+
+	TSharedPtr<FAssetThumbnailPool> ThumbnailPool;
+};
+
+
+#undef LOCTEXT_NAMESPACE
diff --git a/Engine/Plugins/Editor/AssetSearch/Source/Private/Widgets/SearchModel.h b/Engine/Plugins/Editor/AssetSearch/Source/Private/Widgets/SearchModel.h
new file mode 100644
--- /dev/null
+++ b/Engine/Plugins/Editor/AssetSearch/Source/Private/Widgets/SearchModel.h
@@ -0,0 +1,169 @@
+// Copyright Epic Games, Inc. All Rights Reserved.
+
+#pragma once
+
+#include "CoreMinimal.h"
+#include "UObject/ObjectMacros.h"
+#include "UObject/Class.h"
+#include "UObject/WeakObjectPtrTemplates.h"
+#include "UObject/Package.h"
+#include "Engine/World.h"
+#include "Templates/SharedPointer.h"
+
+enum class ESearchNodeType : uint8
+{
+	Asset,
+	Object,
+	Property
+};
+
+class FSearchNode : public TSharedFromThis<FSearchNode>
+{
+public:
+	virtual ~FSearchNode() { }
+	virtual void GetChildren(TArray<TSharedPtr<FSearchNode>>& OutChildren) { }
+	virtual ESearchNodeType GetType() const = 0;
+	virtual FString GetText() const = 0;
+	virtual FString GetObjectPath() const = 0;
+
+protected:
+	float Score = 0;
+};
+
+class FAssetObjectPropertyNode : public FSearchNode
+{
+public:
+	FAssetObjectPropertyNode(const FSearchRecord& InResult)
+		: Record(InResult)
+	{
+	}
+
+	virtual ESearchNodeType GetType() const override { return ESearchNodeType::Property; }
+	virtual FString GetText() const override
+	{
+		return Record.property_name + TEXT(" = ") +  Record.value_text.Replace(TEXT("\n"), TEXT(" "));
+	}
+
+	virtual FString GetObjectPath() const override
+	{
+		return Record.object_path;
+	}
+
+	FSearchRecord Record;
+};
+
+class FAssetObjectNode : public FSearchNode
+{
+public:
+	FAssetObjectNode(const FSearchRecord& InResult)
+		: object_name(InResult.object_name)
+		, object_path(InResult.object_path)
+		, object_native_class(InResult.object_native_class)
+	{
+		Append(InResult);
+	}
+
+	void Append(const FSearchRecord& InResult)
+	{
+		Properties.Add(MakeShared<FAssetObjectPropertyNode>(InResult));
+		Score += InResult.Score;
+	}
+
+	virtual ESearchNodeType GetType() const override { return ESearchNodeType::Object; }
+
+	virtual void GetChildren(TArray<TSharedPtr<FSearchNode>>& OutChildren) override
+	{
+		for (TSharedPtr<FAssetObjectPropertyNode>& Property : Properties)
+		{
+			OutChildren.Add(Property);
+		}
+	}
+
+	virtual FString GetText() const override
+	{
+		return object_name;
+	}
+
+	virtual FString GetObjectPath() const override
+	{
+		return object_path;
+	}
+
+	FString object_name;
+	FString object_path;
+	FString object_native_class;
+
+	TArray<TSharedPtr<FAssetObjectPropertyNode>> Properties;
+};
+
+class FAssetNode : public FSearchNode
+{
+public:
+	FAssetNode(const FSearchRecord& InResult)
+		: AssetName(InResult.AssetName)
+		, AssetPath(InResult.AssetPath)
+		, AssetClass(InResult.AssetClass)
+	{
+		Append(InResult);
+	}
+
+	void Append(const FSearchRecord& InResult)
+	{
+		if (InResult.object_path == AssetPath)
+		{
+			Properties.Add(MakeShared<FAssetObjectPropertyNode>(InResult));
+			Score += InResult.Score;
+			return;
+		}
+
+		TSharedPtr<FAssetObjectNode> ExistingObject = Objects.FindRef(InResult.object_path);
+		if (!ExistingObject.IsValid())
+		{
+			ExistingObject = MakeShared<FAssetObjectNode>(InResult);
+			Objects.Add(InResult.object_path, ExistingObject);
+		}
+		else
+		{
+			ExistingObject->Append(InResult);
+		}
+
+		Score += InResult.Score;
+
+		SortedObjectArray.Reset();
+		Objects.GenerateValueArray(SortedObjectArray);
+	}
+
+	virtual ESearchNodeType GetType() const override { return ESearchNodeType::Asset; }
+
+	virtual void GetChildren(TArray<TSharedPtr<FSearchNode>>& OutChildren) override
+	{
+		for (TSharedPtr<FAssetObjectPropertyNode>& Property : Properties)
+		{
+			OutChildren.Add(Property);
+		}
+
+		for (TSharedPtr<FAssetObjectNode>& Object : SortedObjectArray)
+		{
+			OutChildren.Add(Object);
+		}
+	}
+
+	virtual FString GetText() const override
+	{
+		return AssetName;
+	}
+
+	virtual FString GetObjectPath() const override
+	{
+		return AssetPath;
+	}
+
+	FString AssetName;
+	FString AssetPath;
+	FString AssetClass;
+
+	TArray<TSharedPtr<FAssetObjectPropertyNode>> Properties;
+	TMap<FString, TSharedPtr<FAssetObjectNode>> Objects;
+
+	TArray<TSharedPtr<FAssetObjectNode>> SortedObjectArray;
+};
diff --git a/Engine/Plugins/Editor/AssetSearch/Source/Public/IAssetSearchModule.h b/Engine/Plugins/Editor/AssetSearch/Source/Public/IAssetSearchModule.h
new file mode 100644
--- /dev/null
+++ b/Engine/Plugins/Editor/AssetSearch/Source/Public/IAssetSearchModule.h
@@ -0,0 +1,80 @@
+// Copyright Epic Games, Inc. All Rights Reserved.
+
+#pragma once
+
+#include "CoreMinimal.h"
+#include "Modules/ModuleInterface.h"
+#include "Async/Future.h"
+#include "Async/AsyncResult.h"
+#include "SearchSerializer.h"
+
+struct FSearchQuery
+{
+	FString Query;
+};
+
+struct FSearchRecord
+{
+	FString AssetName;
+	FString AssetPath;
+	FString AssetClass;
+
+	FString object_name;
+	FString object_path;
+	FString object_native_class;
+
+	FString property_name;
+	FString property_field;
+	FString property_class;
+
+	FString value_text;
+	FString value_hidden;
+
+	float Score;
+};
+
+struct FSearchStats
+{
+	int32 Scanning = 0;
+	int32 Downloading = 0;
+	int32 PendingDatabaseUpdates = 0;
+};
+
+class IAssetIndexer
+{
+public:
+	virtual FString GetName() const = 0;
+	virtual int32 GetVersion() const = 0;
+	virtual void IndexAsset(const UObject* InAssetObject, FSearchSerializer& Serializer) = 0;
+};
+
+/**
+ *
+ */
+class IAssetSearchModule : public IModuleInterface
+{
+public:
+	static inline IAssetSearchModule& Get()
+	{
+		static const FName ModuleName = "AssetSearch";
+		return FModuleManager::LoadModuleChecked<IAssetSearchModule>(ModuleName);
+	}
+
+	static inline bool IsAvailable()
+	{
+		static const FName ModuleName = "AssetSearch";
+		return FModuleManager::Get().IsModuleLoaded(ModuleName);
+	}
+
+	virtual FSearchStats GetStats() const = 0;
+
+	virtual bool Search(const FSearchQuery& Query, TFunctionRef<bool(FSearchRecord&&)> InCallback) = 0;
+
+	virtual void RegisterIndexer(FName AssetClassName, IAssetIndexer* Indexer) = 0;
+	virtual void UnregisterIndexer(IAssetIndexer* Indexer) = 0;
+
+public:
+
+	/** Virtual destructor. */
+	virtual ~IAssetSearchModule() { }
+};
diff --git a/Engine/Plugins/Editor/AssetSearch/Source/Public/SearchSerializer.h b/Engine/Plugins/Editor/AssetSearch/Source/Public/SearchSerializer.h
new file mode 100644
--- /dev/null
+++ b/Engine/Plugins/Editor/AssetSearch/Source/Public/SearchSerializer.h
@@ -0,0 +1,69 @@
+// Copyright Epic Games, Inc. All Rights Reserved.
+
+#pragma once
+
+#include "CoreMinimal.h"
+#include "JsonObjectConverter.h"
+#include "Serialization/JsonWriter.h"
+#include "Serialization/JsonSerializer.h"
+#include "Policies/CondensedJsonPrintPolicy.h"
+
+class IAssetIndexer;
+
+class FSearchSerializer : public FNoncopyable
+{
+public:
+	FSearchSerializer(const FAssetData& InAsset, FArchive& InAr);
+	FSearchSerializer(const FAssetData& InAsset, FString* const Stream);
+	~FSearchSerializer();
+
+	static int32 GetVersion();
+
+	void BeginIndexer(IAssetIndexer* InIndexer);
+
+	//---------------------------------------------------------------
+
+	ASSETSEARCH_API void BeginIndexingObject(const UObject* InObject, const FString& InFriendlyName);
+	ASSETSEARCH_API void BeginIndexingObject(const UObject* InObject, const FText& InFriendlyName);
+	ASSETSEARCH_API void EndIndexingObject();
+
+	//---------------------------------------------------------------
+
+	ASSETSEARCH_API void IndexProperty(const FString& InName, const FText& InValue);
+	ASSETSEARCH_API void IndexProperty(const FString& InName, const FString& InValue);
+	ASSETSEARCH_API void IndexProperty(const UClass* InPropertyClass, const FString& InName, const FString& InValue);
+	ASSETSEARCH_API void IndexProperty(const UClass* InPropertyClass, const FString& InName, const FText& InValue);
+	ASSETSEARCH_API void IndexProperty(const FProperty* InProperty, const FString& InValue);
+
+	//---------------------------------------------------------------
+
+	void EndIndexer();
+
+private:
+	void IndexProperty(const FFieldClass& InPropertyFieldClass, const UClass* PropertyClass, const FString& InName, const FString& InValue);
+
+private:
+	void Initialize(const FAssetData& InAsset);
+
+private:
+	TSharedRef< TJsonWriter<TCHAR, TCondensedJsonPrintPolicy<TCHAR> > > JsonWriter;
+
+	FAssetData AssetData;
+
+	const UObject* CurrentObject = nullptr;
+	FString CurrentObjectName;
+	FString CurrentObjectClass;
+	FString CurrentObjectPath;
+
+	struct FIndexedValue
+	{
+		FString PropertyName;
+		FString PropertyFieldClass;
+		FName PropertyClass;
+
+		FString Text;
+		FString HiddenText;
+	};
+
+	TArray<FIndexedValue> Values;
+};
diff --git a/Engine/Plugins/Editor/AssetSearch/Source/Public/Utility/IndexerUtilities.h b/Engine/Plugins/Editor/AssetSearch/Source/Public/Utility/IndexerUtilities.h
new file mode 100644
--- /dev/null
+++ b/Engine/Plugins/Editor/AssetSearch/Source/Public/Utility/IndexerUtilities.h
@@ -0,0 +1,12 @@
+// Copyright Epic Games, Inc. All Rights Reserved.
+
+#pragma once
+
+#include "CoreMinimal.h"
+
+class FIndexerUtilities
+{
+public:
+	static void IterateIndexableProperties(const UObject* InObject, TFunctionRef<void(const FProperty* /*Property*/, const FString& /*Value*/)> Callback);
+	static void IterateIndexableProperties(const UStruct* InStruct, const void* InStructValue, TFunctionRef<void(const FProperty* /*Property*/, const FString& /*Value*/)> Callback);
+};
\ No newline at end of file
diff --git a/Engine/Source/Runtime/Core/Private/Misc/TextFilterExpressionEvaluator.cpp b/Engine/Source/Runtime/Core/Private/Misc/TextFilterExpressionEvaluator.cpp
--- a/Engine/Source/Runtime/Core/Private/Misc/TextFilterExpressionEvaluator.cpp
+++ b/Engine/Source/Runtime/Core/Private/Misc/TextFilterExpressionEvaluator.cpp
@@ -515,7 +515,10 @@ void FTextFilterExpressionEvaluator::SetupGrammar()
 	TokenDefinitions.DefineToken(&ConsumeOperator<FOr>);
 	TokenDefinitions.DefineToken(&ConsumeOperator<FAnd>);
 	TokenDefinitions.DefineToken(&ConsumeOperator<FNot>);
-	TokenDefinitions.DefineToken(&ConsumeNumber);
+	if (ExpressionEvaluatorMode == ETextFilterExpressionEvaluatorMode::Complex)
+	{
+		TokenDefinitions.DefineToken(&ConsumeNumber);
+	}
 	TokenDefinitions.DefineToken((ExpressionEvaluatorMode == ETextFilterExpressionEvaluatorMode::Complex) ? &ConsumeComplexText : &ConsumeBasicText);
 
 	Grammar.DefineGrouping<FSubExpressionStart, FSubExpressionEnd>();
