Tune throttling numbers down while testing.

diff --git a/Engine/Plugins/Editor/AssetSearch/Source/Private/AssetSearchManager.cpp b/Engine/Plugins/Editor/AssetSearch/Source/Private/AssetSearchManager.cpp
--- a/Engine/Plugins/Editor/AssetSearch/Source/Private/AssetSearchManager.cpp
+++ b/Engine/Plugins/Editor/AssetSearch/Source/Private/AssetSearchManager.cpp
@@ -309,7 +309,7 @@ bool FAssetSearchManager::Tick_GameThread(float DeltaTime)
 {
 	check(IsInGameThread());
 
-	int32 ScanLimit = 1000;
+	int32 ScanLimit = 1;
 	while (ProcessAssetQueue.Num() > 0 && ScanLimit > 0)
 	{
 		FAssetData Asset = ProcessAssetQueue.Pop(false);
@@ -317,7 +317,7 @@ bool FAssetSearchManager::Tick_GameThread(float DeltaTime)
 		ScanLimit--;
 	}
 
-	int32 DownloadLimit = 100;
+	int32 DownloadLimit = 1;
 	while (ProcessDDCQueue.Num() > 0 && DownloadLimit > 0)
 	{
 		FScopeLock ScopedLock(&SearchDatabaseCS);
