Build Fixing some unity build issues.

diff --git a/Engine/Plugins/Editor/AssetSearch/Source/Private/AssetSearchManager.h b/Engine/Plugins/Editor/AssetSearch/Source/Private/AssetSearchManager.h
--- a/Engine/Plugins/Editor/AssetSearch/Source/Private/AssetSearchManager.h
+++ b/Engine/Plugins/Editor/AssetSearch/Source/Private/AssetSearchManager.h
@@ -3,6 +3,8 @@
 #pragma once
 
 #include "CoreMinimal.h"
+#include "IAssetSearchModule.h"
+#include "AssetSearchDatabase.h"
 
 class FRunnableThread;
 
diff --git a/Engine/Plugins/Editor/AssetSearch/Source/Private/AssetSearchModule.cpp b/Engine/Plugins/Editor/AssetSearch/Source/Private/AssetSearchModule.cpp
--- a/Engine/Plugins/Editor/AssetSearch/Source/Private/AssetSearchModule.cpp
+++ b/Engine/Plugins/Editor/AssetSearch/Source/Private/AssetSearchModule.cpp
@@ -10,6 +10,7 @@
 #include "Widgets/Docking/SDockTab.h"
 #include "Widgets/SSearchBrowser.h"
 #include "AssetRegistryModule.h"
+#include "AssetSearchManager.h"
 
 #define LOCTEXT_NAMESPACE "FAssetSearchModule"
 
