Tweaking throttling methods for the asset scan.

diff --git a/Engine/Plugins/Editor/AssetSearch/Source/Private/AssetSearchManager.cpp b/Engine/Plugins/Editor/AssetSearch/Source/Private/AssetSearchManager.cpp
--- a/Engine/Plugins/Editor/AssetSearch/Source/Private/AssetSearchManager.cpp
+++ b/Engine/Plugins/Editor/AssetSearch/Source/Private/AssetSearchManager.cpp
@@ -23,6 +23,7 @@ FAutoConsoleVariableRef CVarIndexUnindexAssetsOnLoad(
 FAssetSearchManager::FAssetSearchManager()
 {
 	PendingDatabaseUpdates = 0;
+	PendingDownloads = 0;
 }
 
 FAssetSearchManager::~FAssetSearchManager()
@@ -64,7 +65,7 @@ FSearchStats FAssetSearchManager::GetStats() const
 {
 	FSearchStats Stats;
 	Stats.Scanning = ProcessAssetQueue.Num();
-	Stats.Downloading = ProcessDDCQueue.Num();
+	Stats.Downloading = PendingDownloads;
 	Stats.PendingDatabaseUpdates = PendingDatabaseUpdates;
 	return Stats;
 }
@@ -184,13 +185,14 @@ FString FAssetSearchManager::TryGetDDCKeyForAsset(const FAssetData& InAsset)
 	return AssetJsonDDCKey;
 }
 
-void FAssetSearchManager::TryLoadIndexForAsset(const FAssetData& InAsset)
+bool FAssetSearchManager::TryLoadIndexForAsset(const FAssetData& InAsset)
 {
 	check(IsInGameThread());
 	FString AssetJsonDDCKey = TryGetDDCKeyForAsset(InAsset);
 
 	if (!AssetJsonDDCKey.IsEmpty())
 	{
+		PendingDownloads++;
 		AsyncTask(ENamedThreads::AnyBackgroundThreadNormalTask, [this, InAsset, AssetJsonDDCKey]() {
 			FScopeLock ScopedLock(&SearchDatabaseCS);
 			if (!SearchDatabase.IsAssetUpToDate(InAsset, AssetJsonDDCKey))
@@ -202,7 +204,11 @@ void FAssetSearchManager::TryLoadIndexForAsset(const FAssetData& InAsset)
 				ProcessDDCQueue.Add(DDCRequest);
 			}
 		});
+
+		return true;
 	}
+
+	return false;
 }
 
 FString FAssetSearchManager::GetDerivedDataKey(const FSHAHash& IndexedContentHash)
@@ -310,15 +316,19 @@ bool FAssetSearchManager::Tick_GameThread(float DeltaTime)
 	check(IsInGameThread());
 
 	int32 ScanLimit = 1;
-	while (ProcessAssetQueue.Num() > 0 && ScanLimit > 0)
+	while (ProcessAssetQueue.Num() > 0 && ScanLimit > 0 && PendingDownloads < 100)
 	{
 		FAssetData Asset = ProcessAssetQueue.Pop(false);
-		TryLoadIndexForAsset(Asset);
+		if (TryLoadIndexForAsset(Asset))
+		{
+			ScanLimit -= 100;
+		}
+
 		ScanLimit--;
 	}
 
-	int32 DownloadLimit = 1;
-	while (ProcessDDCQueue.Num() > 0 && DownloadLimit > 0)
+	int32 DownloadProcessLimit = 1;
+	while (ProcessDDCQueue.Num() > 0 && DownloadProcessLimit > 0)
 	{
 		FScopeLock ScopedLock(&SearchDatabaseCS);
 
@@ -335,7 +345,8 @@ bool FAssetSearchManager::Tick_GameThread(float DeltaTime)
 			}
 
 			ProcessDDCQueue.RemoveAt(0, 1, false);
-			DownloadLimit--;
+			PendingDownloads--;
+			DownloadProcessLimit--;
 			continue;
 		}
 		break;
diff --git a/Engine/Plugins/Editor/AssetSearch/Source/Private/AssetSearchManager.h b/Engine/Plugins/Editor/AssetSearch/Source/Private/AssetSearchManager.h
--- a/Engine/Plugins/Editor/AssetSearch/Source/Private/AssetSearchManager.h
+++ b/Engine/Plugins/Editor/AssetSearch/Source/Private/AssetSearchManager.h
@@ -35,7 +35,7 @@ private:
 
 	void RequestIndexAsset(UObject* InAsset);
 	bool IsAssetIndexable(UObject* InAsset);
-	void TryLoadIndexForAsset(const FAssetData& InAsset);
+	bool TryLoadIndexForAsset(const FAssetData& InAsset);
 	FString TryGetDDCKeyForAsset(const FAssetData& InAsset);
 	FString GetDerivedDataKey(const FSHAHash& IndexedContentHash);
 	FString GetDerivedDataKey(const FAssetData& UnindexedAsset);
@@ -55,10 +55,13 @@ private:
 	FAssetSearchDatabase SearchDatabase;
 	FCriticalSection SearchDatabaseCS;
 	TAtomic<int32> PendingDatabaseUpdates;
+	TAtomic<int32> PendingDownloads;
 
 private:
 	TMap<FName, IAssetIndexer*> Indexers;
 
+	TArray<TWeakObjectPtr<UObject>> RequestIndexQueue;
+
 	TArray<FAssetData> ProcessAssetQueue;
 
 	struct FAssetDDCRequest
