Exposing different throttling variables.

diff --git a/Engine/Plugins/Editor/AssetSearch/Source/Private/AssetSearchManager.cpp b/Engine/Plugins/Editor/AssetSearch/Source/Private/AssetSearchManager.cpp
--- a/Engine/Plugins/Editor/AssetSearch/Source/Private/AssetSearchManager.cpp
+++ b/Engine/Plugins/Editor/AssetSearch/Source/Private/AssetSearchManager.cpp
@@ -20,6 +20,14 @@ FAutoConsoleVariableRef CVarIndexUnindexAssetsOnLoad(
 	TEXT("Index Unindex Assets On Load")
 );
 
+static int32 PendingDownloadsMax = 25;
+FAutoConsoleVariableRef CVarPendingDownloadsMax(
+	TEXT("Search.PendingDownloadsMax"),
+	PendingDownloadsMax,
+	TEXT("")
+);
+
+
 FAssetSearchManager::FAssetSearchManager()
 {
 	PendingDatabaseUpdates = 0;
@@ -315,19 +323,19 @@ bool FAssetSearchManager::Tick_GameThread(float DeltaTime)
 {
 	check(IsInGameThread());
 
-	int32 ScanLimit = 1;
-	while (ProcessAssetQueue.Num() > 0 && ScanLimit > 0 && PendingDownloads < 100)
+	int32 ScanLimit = 1000;
+	while (ProcessAssetQueue.Num() > 0 && ScanLimit > 0 && PendingDownloads < PendingDownloadsMax)
 	{
 		FAssetData Asset = ProcessAssetQueue.Pop(false);
 		if (TryLoadIndexForAsset(Asset))
 		{
-			ScanLimit -= 100;
+			break;
 		}
 
 		ScanLimit--;
 	}
 
-	int32 DownloadProcessLimit = 1;
+	int32 DownloadProcessLimit = 10;
 	while (ProcessDDCQueue.Num() > 0 && DownloadProcessLimit > 0)
 	{
 		FScopeLock ScopedLock(&SearchDatabaseCS);
