Fixing some monolithic build issues.

diff --git a/Engine/Plugins/Editor/AssetSearch/Source/Public/IAssetSearchModule.h b/Engine/Plugins/Editor/AssetSearch/Source/Public/IAssetSearchModule.h
--- a/Engine/Plugins/Editor/AssetSearch/Source/Public/IAssetSearchModule.h
+++ b/Engine/Plugins/Editor/AssetSearch/Source/Public/IAssetSearchModule.h
@@ -7,6 +7,7 @@
 #include "Async/Future.h"
 #include "Async/AsyncResult.h"
 #include "SearchSerializer.h"
+#include "Modules/ModuleManager.h"
 
 struct FSearchQuery
 {
diff --git a/Engine/Plugins/Editor/AssetSearch/Source/Public/SearchSerializer.h b/Engine/Plugins/Editor/AssetSearch/Source/Public/SearchSerializer.h
--- a/Engine/Plugins/Editor/AssetSearch/Source/Public/SearchSerializer.h
+++ b/Engine/Plugins/Editor/AssetSearch/Source/Public/SearchSerializer.h
@@ -7,6 +7,8 @@
 #include "Serialization/JsonWriter.h"
 #include "Serialization/JsonSerializer.h"
 #include "Policies/CondensedJsonPrintPolicy.h"
+#include "Serialization/Archive.h"
+#include "AssetData.h"
 
 class IAssetIndexer;
 
