Additional monolithic issues.

diff --git a/Engine/Plugins/Editor/AssetSearch/Source/Private/SearchSerializer.cpp b/Engine/Plugins/Editor/AssetSearch/Source/Private/SearchSerializer.cpp
--- a/Engine/Plugins/Editor/AssetSearch/Source/Private/SearchSerializer.cpp
+++ b/Engine/Plugins/Editor/AssetSearch/Source/Private/SearchSerializer.cpp
@@ -5,6 +5,8 @@
 #include "Serialization/JsonWriter.h"
 #include "Serialization/JsonSerializer.h"
 #include "Policies/CondensedJsonPrintPolicy.h"
+#include "UObject/UnrealType.h"
+#include "IAssetSearchModule.h"
 
 #define LOCTEXT_NAMESPACE "FSearchSerializer"
 
