Fixing a hitch on the game thread for the searcher.

diff --git a/Engine/Plugins/Editor/AssetSearch/Source/Private/AssetSearchManager.cpp b/Engine/Plugins/Editor/AssetSearch/Source/Private/AssetSearchManager.cpp
--- a/Engine/Plugins/Editor/AssetSearch/Source/Private/AssetSearchManager.cpp
+++ b/Engine/Plugins/Editor/AssetSearch/Source/Private/AssetSearchManager.cpp
@@ -20,13 +20,26 @@ FAutoConsoleVariableRef CVarIndexUnindexAssetsOnLoad(
 	TEXT("Index Unindex Assets On Load")
 );
 
-static int32 PendingDownloadsMax = 25;
+static int32 PendingDownloadsMax = 100;
 FAutoConsoleVariableRef CVarPendingDownloadsMax(
 	TEXT("Search.PendingDownloadsMax"),
 	PendingDownloadsMax,
 	TEXT("")
 );
 
+static int32 GameThread_DownloadProcessLimit = 30;
+FAutoConsoleVariableRef CVarGameThread_DownloadProcessLimit(
+	TEXT("Search.GameThread_DownloadProcessLimit"),
+	GameThread_DownloadProcessLimit,
+	TEXT("")
+);
+
+static int32 GameThread_AssetScanLimit = 1000;
+FAutoConsoleVariableRef CVarGameThread_AssetScanLimit(
+	TEXT("Search.GameThread_AssetScanLimit"),
+	GameThread_AssetScanLimit,
+	TEXT("")
+);
 
 FAssetSearchManager::FAssetSearchManager()
 {
@@ -209,7 +222,7 @@ bool FAssetSearchManager::TryLoadIndexForAsset(const FAssetData& InAsset)
 				DDCRequest.AssetData = InAsset;
 				DDCRequest.DDCKey_IndexDataHash = AssetJsonDDCKey;
 				DDCRequest.DDCHandle = GetDerivedDataCacheRef().GetAsynchronous(*AssetJsonDDCKey, InAsset.ObjectPath.ToString());
-				ProcessDDCQueue.Add(DDCRequest);
+				ProcessDDCQueue.Enqueue(DDCRequest);
 			}
 		});
 
@@ -323,36 +336,34 @@ bool FAssetSearchManager::Tick_GameThread(float DeltaTime)
 {
 	check(IsInGameThread());
 
-	int32 ScanLimit = 1000;
+	int32 ScanLimit = GameThread_AssetScanLimit;
 	while (ProcessAssetQueue.Num() > 0 && ScanLimit > 0 && PendingDownloads < PendingDownloadsMax)
 	{
 		FAssetData Asset = ProcessAssetQueue.Pop(false);
 		if (TryLoadIndexForAsset(Asset))
 		{
-			break;
+			ScanLimit -= 10;
 		}
 
 		ScanLimit--;
 	}
 
-	int32 DownloadProcessLimit = 10;
-	while (ProcessDDCQueue.Num() > 0 && DownloadProcessLimit > 0)
+	int32 DownloadProcessLimit = GameThread_DownloadProcessLimit;
+	while (!ProcessDDCQueue.IsEmpty() && DownloadProcessLimit > 0)
 	{
-		FScopeLock ScopedLock(&SearchDatabaseCS);
-
-		const FAssetDDCRequest& PendingRequest = ProcessDDCQueue[0];
-		if (GetDerivedDataCacheRef().PollAsynchronousCompletion(PendingRequest.DDCHandle))
+		const FAssetDDCRequest* PendingRequest = ProcessDDCQueue.Peek();
+		if (GetDerivedDataCacheRef().PollAsynchronousCompletion(PendingRequest->DDCHandle))
 		{
 			bool bDataWasBuilt;
 
 			TArray<uint8> OutContent;
-			bool bGetSuccessful = GetDerivedDataCacheRef().GetAsynchronousResults(PendingRequest.DDCHandle, OutContent, &bDataWasBuilt);
+			bool bGetSuccessful = GetDerivedDataCacheRef().GetAsynchronousResults(PendingRequest->DDCHandle, OutContent, &bDataWasBuilt);
 			if (bGetSuccessful)
 			{
-				LoadDDCContentIntoDatabase(PendingRequest.AssetData, OutContent, PendingRequest.DDCKey_IndexDataHash);
+				LoadDDCContentIntoDatabase(PendingRequest->AssetData, OutContent, PendingRequest->DDCKey_IndexDataHash);
 			}
 
-			ProcessDDCQueue.RemoveAt(0, 1, false);
+			ProcessDDCQueue.Pop();
 			PendingDownloads--;
 			DownloadProcessLimit--;
 			continue;
diff --git a/Engine/Plugins/Editor/AssetSearch/Source/Private/AssetSearchManager.h b/Engine/Plugins/Editor/AssetSearch/Source/Private/AssetSearchManager.h
--- a/Engine/Plugins/Editor/AssetSearch/Source/Private/AssetSearchManager.h
+++ b/Engine/Plugins/Editor/AssetSearch/Source/Private/AssetSearchManager.h
@@ -5,6 +5,7 @@
 #include "CoreMinimal.h"
 #include "IAssetSearchModule.h"
 #include "AssetSearchDatabase.h"
+#include "Containers/Queue.h"
 
 class FRunnableThread;
 
@@ -70,7 +71,7 @@ private:
 		FString DDCKey_IndexDataHash;
 		uint32 DDCHandle;
 	};
-	TArray<FAssetDDCRequest> ProcessDDCQueue;
+	TQueue<FAssetDDCRequest, EQueueMode::Mpsc> ProcessDDCQueue;
 
 	FDelegateHandle TickerHandle;
 
diff --git a/Engine/Plugins/Editor/AssetSearch/Source/Private/Widgets/SSearchBrowser.h b/Engine/Plugins/Editor/AssetSearch/Source/Private/Widgets/SSearchBrowser.h
--- a/Engine/Plugins/Editor/AssetSearch/Source/Private/Widgets/SSearchBrowser.h
+++ b/Engine/Plugins/Editor/AssetSearch/Source/Private/Widgets/SSearchBrowser.h
@@ -6,6 +6,9 @@
 #include "Widgets/Views/SHeaderRow.h"
 #include "Widgets/Views/SListView.h"
 #include "Widgets/Views/STableViewBase.h"
+#include "SearchModel.h"
+#include "AssetThumbnail.h"
+#include "Widgets/Views/STreeView.h"
 
 class IDetailsView;
 class ITableRow;
diff --git a/Engine/Plugins/Editor/AssetSearch/Source/Private/Widgets/SSearchTreeRow.h b/Engine/Plugins/Editor/AssetSearch/Source/Private/Widgets/SSearchTreeRow.h
--- a/Engine/Plugins/Editor/AssetSearch/Source/Private/Widgets/SSearchTreeRow.h
+++ b/Engine/Plugins/Editor/AssetSearch/Source/Private/Widgets/SSearchTreeRow.h
@@ -11,6 +11,7 @@
 #include "Widgets/DeclarativeSyntaxSupport.h"
 #include "Input/Reply.h"
 #include "Widgets/Views/STableViewBase.h"
+#include "AssetThumbnail.h"
 
 #define LOCTEXT_NAMESPACE "SObjectBrowserTableRow"
 
diff --git a/Engine/Plugins/Editor/AssetSearch/Source/Private/Widgets/SearchModel.h b/Engine/Plugins/Editor/AssetSearch/Source/Private/Widgets/SearchModel.h
--- a/Engine/Plugins/Editor/AssetSearch/Source/Private/Widgets/SearchModel.h
+++ b/Engine/Plugins/Editor/AssetSearch/Source/Private/Widgets/SearchModel.h
@@ -9,6 +9,7 @@
 #include "UObject/Package.h"
 #include "Engine/World.h"
 #include "Templates/SharedPointer.h"
+#include "IAssetSearchModule.h"
 
 enum class ESearchNodeType : uint8
 {
