Fixing some iwyu issues.

diff --git a/Engine/Plugins/Editor/AssetSearch/Source/Private/AssetSearchManager.cpp b/Engine/Plugins/Editor/AssetSearch/Source/Private/AssetSearchManager.cpp
--- a/Engine/Plugins/Editor/AssetSearch/Source/Private/AssetSearchManager.cpp
+++ b/Engine/Plugins/Editor/AssetSearch/Source/Private/AssetSearchManager.cpp
@@ -12,6 +12,8 @@
 #include "Indexers/BlueprintIndexer.h"
 #include "Indexers/WidgetBlueprintIndexer.h"
 #include "Containers/StringConv.h"
+#include "Containers/Ticker.h"
+#include "Misc/Paths.h"
 
 static bool bIndexUnindexAssetsOnLoad = false;
 FAutoConsoleVariableRef CVarIndexUnindexAssetsOnLoad(
diff --git a/Engine/Plugins/Editor/AssetSearch/Source/Private/Indexers/WidgetBlueprintIndexer.cpp b/Engine/Plugins/Editor/AssetSearch/Source/Private/Indexers/WidgetBlueprintIndexer.cpp
--- a/Engine/Plugins/Editor/AssetSearch/Source/Private/Indexers/WidgetBlueprintIndexer.cpp
+++ b/Engine/Plugins/Editor/AssetSearch/Source/Private/Indexers/WidgetBlueprintIndexer.cpp
@@ -4,6 +4,7 @@
 #include "Engine/Blueprint.h"
 #include "WidgetBlueprint.h"
 #include "Components/Widget.h"
+#include "Utility/IndexerUtilities.h"
 
 #define LOCTEXT_NAMESPACE "FWidgetBlueprintIndexer"
 
diff --git a/Engine/Plugins/Editor/AssetSearch/Source/Private/Utility/IndexerUtilities.cpp b/Engine/Plugins/Editor/AssetSearch/Source/Private/Utility/IndexerUtilities.cpp
--- a/Engine/Plugins/Editor/AssetSearch/Source/Private/Utility/IndexerUtilities.cpp
+++ b/Engine/Plugins/Editor/AssetSearch/Source/Private/Utility/IndexerUtilities.cpp
@@ -4,6 +4,7 @@
 #include "UObject/UnrealType.h"
 #include "GameplayTagContainer.h"
 #include "Internationalization/Text.h"
+#include "UObject/TextProperty.h"
 
 PRAGMA_DISABLE_OPTIMIZATION
 
diff --git a/Engine/Plugins/Editor/AssetSearch/Source/Private/Widgets/SSearchTreeRow.h b/Engine/Plugins/Editor/AssetSearch/Source/Private/Widgets/SSearchTreeRow.h
--- a/Engine/Plugins/Editor/AssetSearch/Source/Private/Widgets/SSearchTreeRow.h
+++ b/Engine/Plugins/Editor/AssetSearch/Source/Private/Widgets/SSearchTreeRow.h
@@ -16,6 +16,7 @@
 #define LOCTEXT_NAMESPACE "SObjectBrowserTableRow"
 
 class FSearchNode;
+class IAssetRegistry;
 
 /**
  * Implements a row widget for a live uobject.
