Decrement the downloading counter if we fail to start downloading.

diff --git a/Engine/Plugins/Editor/AssetSearch/Source/Private/AssetSearchManager.cpp b/Engine/Plugins/Editor/AssetSearch/Source/Private/AssetSearchManager.cpp
--- a/Engine/Plugins/Editor/AssetSearch/Source/Private/AssetSearchManager.cpp
+++ b/Engine/Plugins/Editor/AssetSearch/Source/Private/AssetSearchManager.cpp
@@ -226,6 +226,10 @@ bool FAssetSearchManager::TryLoadIndexForAsset(const FAssetData& InAsset)
 				DDCRequest.DDCHandle = GetDerivedDataCacheRef().GetAsynchronous(*AssetJsonDDCKey, InAsset.ObjectPath.ToString());
 				ProcessDDCQueue.Enqueue(DDCRequest);
 			}
+			else
+			{
+				PendingDownloads--;
+			}
 		});
 
 		return true;
