Fixing some iwyu issues.

diff --git a/Engine/Plugins/Editor/AssetSearch/Source/Private/SearchSerializer.cpp b/Engine/Plugins/Editor/AssetSearch/Source/Private/SearchSerializer.cpp
--- a/Engine/Plugins/Editor/AssetSearch/Source/Private/SearchSerializer.cpp
+++ b/Engine/Plugins/Editor/AssetSearch/Source/Private/SearchSerializer.cpp
@@ -7,6 +7,7 @@
 #include "Policies/CondensedJsonPrintPolicy.h"
 #include "UObject/UnrealType.h"
 #include "IAssetSearchModule.h"
+#include "UObject/TextProperty.h"
 
 #define LOCTEXT_NAMESPACE "FSearchSerializer"
 
