Adding a way to get the number of total search records that gets polled occasionally.

diff --git a/Engine/Plugins/Editor/AssetSearch/Source/Private/AssetSearchDatabase.cpp b/Engine/Plugins/Editor/AssetSearch/Source/Private/AssetSearchDatabase.cpp
--- a/Engine/Plugins/Editor/AssetSearch/Source/Private/AssetSearchDatabase.cpp
+++ b/Engine/Plugins/Editor/AssetSearch/Source/Private/AssetSearchDatabase.cpp
@@ -49,6 +49,7 @@ public:
 
 		PREPARE_STATEMENT(Statement_GetAssetIdForAssetPath);
 		PREPARE_STATEMENT(Statement_IsAssetUpToDate);
+		PREPARE_STATEMENT(Statement_GetTotalSearchRecords);
 		PREPARE_STATEMENT(Statement_AddAssetToAssetTable);
 		PREPARE_STATEMENT(Statement_AddAssetProperty);
 		PREPARE_STATEMENT(Statement_DeleteEntriesForAsset);
@@ -109,6 +110,23 @@ public:
 		return false;
 	}
 
+	 SQLITE_PREPARED_STATEMENT(FGetTotalSearchRecords,
+		"SELECT COUNT(rowid) FROM table_asset_properties;",
+		SQLITE_PREPARED_STATEMENT_COLUMNS(int64),
+		SQLITE_PREPARED_STATEMENT_BINDINGS()
+	);
+	private: FGetTotalSearchRecords Statement_GetTotalSearchRecords;
+	public: int64 GetTotalSearchRecords()
+	{
+		int64 OutTotalSearchRecords;
+		if (Statement_GetTotalSearchRecords.BindAndExecuteSingle(OutTotalSearchRecords))
+		{
+			return OutTotalSearchRecords;
+		}
+
+		return INDEX_NONE;
+	}
+
 	SQLITE_PREPARED_STATEMENT(FGetAssetIdForAssetPath,
 		"SELECT assetid FROM table_assets WHERE asset_path = ?1;",
 		SQLITE_PREPARED_STATEMENT_COLUMNS(int64),
@@ -691,4 +709,14 @@ bool FAssetSearchDatabase::EnumerateSearchResults(const FSearchQuery& Query, TFu
 	return bSuccess;
 }
 
+int64 FAssetSearchDatabase::GetTotalSearchRecords() const
+{
+	if (ensure(Statements))
+	{
+		return Statements->GetTotalSearchRecords();
+	}
+
+	return INDEX_NONE;
+}
+
 PRAGMA_ENABLE_OPTIMIZATION
diff --git a/Engine/Plugins/Editor/AssetSearch/Source/Private/AssetSearchDatabase.h b/Engine/Plugins/Editor/AssetSearch/Source/Private/AssetSearchDatabase.h
--- a/Engine/Plugins/Editor/AssetSearch/Source/Private/AssetSearchDatabase.h
+++ b/Engine/Plugins/Editor/AssetSearch/Source/Private/AssetSearchDatabase.h
@@ -79,6 +79,11 @@ public:
 	 */
 	bool EnumerateSearchResults(const FSearchQuery& Query, TFunctionRef<bool(FSearchRecord&&)> InCallback);
 
+	/**
+	 * Get the total number of records.
+	 */
+	int64 GetTotalSearchRecords() const;
+
 private:
 	void LogLastError() const;
 
diff --git a/Engine/Plugins/Editor/AssetSearch/Source/Private/AssetSearchManager.cpp b/Engine/Plugins/Editor/AssetSearch/Source/Private/AssetSearchManager.cpp
--- a/Engine/Plugins/Editor/AssetSearch/Source/Private/AssetSearchManager.cpp
+++ b/Engine/Plugins/Editor/AssetSearch/Source/Private/AssetSearchManager.cpp
@@ -47,6 +47,8 @@ FAssetSearchManager::FAssetSearchManager()
 {
 	PendingDatabaseUpdates = 0;
 	PendingDownloads = 0;
+	TotalSearchRecords = 0;
+	LastRecordCountUpdateSeconds = 0;
 }
 
 FAssetSearchManager::~FAssetSearchManager()
@@ -90,6 +92,7 @@ FSearchStats FAssetSearchManager::GetStats() const
 	Stats.Scanning = ProcessAssetQueue.Num();
 	Stats.Downloading = PendingDownloads;
 	Stats.PendingDatabaseUpdates = PendingDatabaseUpdates;
+	Stats.TotalRecords = TotalSearchRecords;
 	return Stats;
 }
 
@@ -377,6 +380,16 @@ bool FAssetSearchManager::Tick_GameThread(float DeltaTime)
 		break;
 	}
 
+	if ((FPlatformTime::Seconds() - LastRecordCountUpdateSeconds) > 30)
+	{
+		LastRecordCountUpdateSeconds = FPlatformTime::Seconds();
+
+		AsyncTask(ENamedThreads::AnyBackgroundThreadNormalTask, [this]() {
+			FScopeLock ScopedLock(&SearchDatabaseCS);
+			TotalSearchRecords = SearchDatabase.GetTotalSearchRecords();
+		});
+	}
+
 	return true;
 }
 
diff --git a/Engine/Plugins/Editor/AssetSearch/Source/Private/AssetSearchManager.h b/Engine/Plugins/Editor/AssetSearch/Source/Private/AssetSearchManager.h
--- a/Engine/Plugins/Editor/AssetSearch/Source/Private/AssetSearchManager.h
+++ b/Engine/Plugins/Editor/AssetSearch/Source/Private/AssetSearchManager.h
@@ -57,6 +57,9 @@ private:
 	FCriticalSection SearchDatabaseCS;
 	TAtomic<int32> PendingDatabaseUpdates;
 	TAtomic<int32> PendingDownloads;
+	TAtomic<int64> TotalSearchRecords;
+
+	double LastRecordCountUpdateSeconds;
 
 private:
 	TMap<FName, IAssetIndexer*> Indexers;
diff --git a/Engine/Plugins/Editor/AssetSearch/Source/Private/Widgets/SSearchBrowser.cpp b/Engine/Plugins/Editor/AssetSearch/Source/Private/Widgets/SSearchBrowser.cpp
--- a/Engine/Plugins/Editor/AssetSearch/Source/Private/Widgets/SSearchBrowser.cpp
+++ b/Engine/Plugins/Editor/AssetSearch/Source/Private/Widgets/SSearchBrowser.cpp
@@ -173,7 +173,7 @@ FText SSearchBrowser::GetStatusText() const
 {
 	IAssetSearchModule& SearchModule = IAssetSearchModule::Get();
 	FSearchStats SearchStats = SearchModule.GetStats();
-	return FText::Format(LOCTEXT("SearchStatusTextFmt", "Scanning {0}   Downloading {1}   Updating {2}"), SearchStats.Scanning, SearchStats.Downloading, SearchStats.PendingDatabaseUpdates);
+	return FText::Format(LOCTEXT("SearchStatusTextFmt", "Scanning {0}   Downloading {1}   Updating {2}            Total Records {3}"), SearchStats.Scanning, SearchStats.Downloading, SearchStats.PendingDatabaseUpdates, SearchStats.TotalRecords);
 }
 
 FReply SSearchBrowser::OnRefresh()
diff --git a/Engine/Plugins/Editor/AssetSearch/Source/Public/IAssetSearchModule.h b/Engine/Plugins/Editor/AssetSearch/Source/Public/IAssetSearchModule.h
--- a/Engine/Plugins/Editor/AssetSearch/Source/Public/IAssetSearchModule.h
+++ b/Engine/Plugins/Editor/AssetSearch/Source/Public/IAssetSearchModule.h
@@ -39,6 +39,8 @@ struct FSearchStats
 	int32 Scanning = 0;
 	int32 Downloading = 0;
 	int32 PendingDatabaseUpdates = 0;
+
+	int64 TotalRecords = 0;
 };
 
 class IAssetIndexer
