Fixing a linux compiler warning

diff --git a/Engine/Plugins/Editor/AssetSearch/Source/Private/AssetSearchDatabase.h b/Engine/Plugins/Editor/AssetSearch/Source/Private/AssetSearchDatabase.h
--- a/Engine/Plugins/Editor/AssetSearch/Source/Private/AssetSearchDatabase.h
+++ b/Engine/Plugins/Editor/AssetSearch/Source/Private/AssetSearchDatabase.h
@@ -92,7 +92,7 @@ private:
 	TUniquePtr<FSQLiteDatabase> Database;
 
 	/** The database filename name we use */
-	const FString DatabaseFileName;
+	FString DatabaseFileName;
 
 	/** Root path to store all session data under */
 	FString SessionPath;
