When searching we now sort by relevance/best match, we also search while typing.

diff --git a/Engine/Plugins/Editor/AssetSearch/Source/Private/AssetSearchManager.cpp b/Engine/Plugins/Editor/AssetSearch/Source/Private/AssetSearchManager.cpp
--- a/Engine/Plugins/Editor/AssetSearch/Source/Private/AssetSearchManager.cpp
+++ b/Engine/Plugins/Editor/AssetSearch/Source/Private/AssetSearchManager.cpp
@@ -393,10 +393,24 @@ bool FAssetSearchManager::Tick_GameThread(float DeltaTime)
 	return true;
 }
 
-bool FAssetSearchManager::Search(const FSearchQuery& Query, TFunctionRef<bool(FSearchRecord&&)> InCallback)
+void FAssetSearchManager::Search(const FSearchQuery& Query, TFunction<void(TArray<FSearchRecord>&&)> InCallback)
 {
 	check(IsInGameThread());
 
-	FScopeLock ScopedLock(&SearchDatabaseCS);
-	return SearchDatabase.EnumerateSearchResults(Query, InCallback);
+	AsyncTask(ENamedThreads::AnyBackgroundThreadNormalTask, [this, Query, InCallback]() {
+
+		TArray<FSearchRecord> Results;
+
+		{
+			FScopeLock ScopedLock(&SearchDatabaseCS);
+			SearchDatabase.EnumerateSearchResults(Query, [&Results](FSearchRecord&& InResult) {
+				Results.Add(MoveTemp(InResult));
+				return true;
+			});
+		}
+
+		AsyncTask(ENamedThreads::GameThread, [ResultsFwd = MoveTemp(Results), InCallback]() mutable {
+			InCallback(MoveTemp(ResultsFwd));
+		});
+	});
 }
\ No newline at end of file
diff --git a/Engine/Plugins/Editor/AssetSearch/Source/Private/AssetSearchManager.h b/Engine/Plugins/Editor/AssetSearch/Source/Private/AssetSearchManager.h
--- a/Engine/Plugins/Editor/AssetSearch/Source/Private/AssetSearchManager.h
+++ b/Engine/Plugins/Editor/AssetSearch/Source/Private/AssetSearchManager.h
@@ -20,7 +20,7 @@ public:
 
 	FSearchStats GetStats() const;
 
-	bool Search(const FSearchQuery& Query, TFunctionRef<bool(FSearchRecord&&)> InCallback);
+	void Search(const FSearchQuery& Query, TFunction<void(TArray<FSearchRecord>&&)> InCallback);
 
 private:
 	bool Tick_GameThread(float DeltaTime);
diff --git a/Engine/Plugins/Editor/AssetSearch/Source/Private/AssetSearchModule.cpp b/Engine/Plugins/Editor/AssetSearch/Source/Private/AssetSearchModule.cpp
--- a/Engine/Plugins/Editor/AssetSearch/Source/Private/AssetSearchModule.cpp
+++ b/Engine/Plugins/Editor/AssetSearch/Source/Private/AssetSearchModule.cpp
@@ -66,9 +66,9 @@ public:
 		return SearchManager->GetStats();
 	}
 
-	virtual bool Search(const FSearchQuery& Query, TFunctionRef<bool(FSearchRecord&&)> InCallback) override
+	virtual void Search(const FSearchQuery& Query, TFunction<void(TArray<FSearchRecord>&&)> InCallback) override
 	{
-		return SearchManager->Search(Query, InCallback);
+		SearchManager->Search(Query, InCallback);
 	}
 
 	virtual void RegisterIndexer(FName AssetClassName, IAssetIndexer* Indexer) override
diff --git a/Engine/Plugins/Editor/AssetSearch/Source/Private/Widgets/SSearchBrowser.cpp b/Engine/Plugins/Editor/AssetSearch/Source/Private/Widgets/SSearchBrowser.cpp
--- a/Engine/Plugins/Editor/AssetSearch/Source/Private/Widgets/SSearchBrowser.cpp
+++ b/Engine/Plugins/Editor/AssetSearch/Source/Private/Widgets/SSearchBrowser.cpp
@@ -81,8 +81,10 @@ void SSearchBrowser::Construct( const FArguments& InArgs )
 				.FillWidth(1.0f)
 				[
 					SNew(SEditableTextBox)
+					.ClearKeyboardFocusOnCommit(false)
 					.HintText(LOCTEXT("SearchHint", "Search"))
 					.OnTextCommitted(this, &SSearchBrowser::OnSearchTextCommited)
+					.OnTextChanged(this, &SSearchBrowser::OnSearchTextChanged)
 				]
 			]
 
@@ -214,9 +216,28 @@ void SSearchBrowser::RefreshList()
 		Query.Query = FilterText.ToString();
 
 		IAssetSearchModule& SearchModule = IAssetSearchModule::Get();
-		SearchModule.Search(Query, [this](FSearchRecord&& InResult) {
-			AppendResult(MoveTemp(InResult));
-			return true;
+		SearchModule.Search(Query, [this](TArray<FSearchRecord>&& InResults) {
+
+			SearchResults.Reset();
+			SearchResultHierarchy.Reset();
+
+			for (int32 ResultIndex = 0; ResultIndex < InResults.Num(); ResultIndex++)
+			{
+				AppendResult(MoveTemp(InResults[ResultIndex]));
+			}
+
+			for (auto& Entry : SearchResultHierarchy)
+			{
+				SearchResults.Add(Entry.Value);
+
+				SearchTreeView->SetItemExpansion(Entry.Value, true);
+			}
+
+			SearchResults.Sort([](const TSharedPtr<FSearchNode>& A, const TSharedPtr<FSearchNode>& B) {
+				return A->GetMaxScore() < B->GetMaxScore();
+			});
+
+			SearchTreeView->RequestListRefresh();
 		});
 	}
 }
@@ -233,23 +254,30 @@ void SSearchBrowser::AppendResult(FSearchRecord&& InResult)
 	{
 		ExistingAssetNode->Append(InResult);
 	}
+}
 
-	SearchResults.Reset();
-	for (auto& Entry : SearchResultHierarchy)
+void SSearchBrowser::OnSearchTextCommited(const FText& InText, ETextCommit::Type InCommitType)
+{
+	TryRefreshingSearch(InText);
+}
+
+void SSearchBrowser::OnSearchTextChanged(const FText& InText)
+{
+	if (InText.ToString().Len() > 3)
 	{
-		SearchResults.Add(Entry.Value);
+		TryRefreshingSearch(InText);
 	}
-
-	SearchTreeView->RequestListRefresh();
-	SearchTreeView->SetItemExpansion(ExistingAssetNode, true);
 }
 
-void SSearchBrowser::OnSearchTextCommited(const FText& InText, ETextCommit::Type InCommitType)
+void SSearchBrowser::TryRefreshingSearch(const FText& InText)
 {
-	FilterText = InText;
-	FilterString = FilterText.ToString();
+	if (FilterText.ToString() != InText.ToString())
+	{
+		FilterText = InText;
+		FilterString = FilterText.ToString();
 
-	RefreshList();
+		RefreshList();
+	}
 }
 
 TSharedRef<ITableRow> SSearchBrowser::HandleListGenerateRow(TSharedPtr<FSearchNode> ObjectPtr, const TSharedRef<STableViewBase>& OwnerTable)
diff --git a/Engine/Plugins/Editor/AssetSearch/Source/Private/Widgets/SSearchBrowser.h b/Engine/Plugins/Editor/AssetSearch/Source/Private/Widgets/SSearchBrowser.h
--- a/Engine/Plugins/Editor/AssetSearch/Source/Private/Widgets/SSearchBrowser.h
+++ b/Engine/Plugins/Editor/AssetSearch/Source/Private/Widgets/SSearchBrowser.h
@@ -51,6 +51,8 @@ private:
 	void AppendResult(FSearchRecord&& InResult);
 
 	void OnSearchTextCommited(const FText& InText, ETextCommit::Type InCommitType);
+	void OnSearchTextChanged(const FText& InText);
+	void TryRefreshingSearch(const FText& InText);
 
 	TSharedRef<ITableRow> HandleListGenerateRow(TSharedPtr<FSearchNode> TransactionInfo, const TSharedRef<STableViewBase>& OwnerTable);
 	void GetChildrenForInfo(TSharedPtr<FSearchNode> InInfo, TArray< TSharedPtr<FSearchNode> >& OutChildren);
diff --git a/Engine/Plugins/Editor/AssetSearch/Source/Private/Widgets/SearchModel.h b/Engine/Plugins/Editor/AssetSearch/Source/Private/Widgets/SearchModel.h
--- a/Engine/Plugins/Editor/AssetSearch/Source/Private/Widgets/SearchModel.h
+++ b/Engine/Plugins/Editor/AssetSearch/Source/Private/Widgets/SearchModel.h
@@ -27,8 +27,12 @@ public:
 	virtual FString GetText() const = 0;
 	virtual FString GetObjectPath() const = 0;
 
+	float GetTotalScore() const { return TotalScore; }
+	float GetMaxScore() const { return MaxScore; }
+
 protected:
-	float Score = 0;
+	float TotalScore = 0;
+	float MaxScore = 0;
 };
 
 class FAssetObjectPropertyNode : public FSearchNode
@@ -67,7 +71,8 @@ public:
 	void Append(const FSearchRecord& InResult)
 	{
 		Properties.Add(MakeShared<FAssetObjectPropertyNode>(InResult));
-		Score += InResult.Score;
+		TotalScore += InResult.Score;
+		MaxScore = FMath::Min(MaxScore, InResult.Score);
 	}
 
 	virtual ESearchNodeType GetType() const override { return ESearchNodeType::Object; }
@@ -113,7 +118,8 @@ public:
 		if (InResult.object_path == AssetPath)
 		{
 			Properties.Add(MakeShared<FAssetObjectPropertyNode>(InResult));
-			Score += InResult.Score;
+			TotalScore += InResult.Score;
+			MaxScore = FMath::Min(MaxScore, InResult.Score);
 			return;
 		}
 
@@ -128,7 +134,8 @@ public:
 			ExistingObject->Append(InResult);
 		}
 
-		Score += InResult.Score;
+		TotalScore += InResult.Score;
+		MaxScore = FMath::Min(MaxScore, InResult.Score);
 
 		SortedObjectArray.Reset();
 		Objects.GenerateValueArray(SortedObjectArray);
diff --git a/Engine/Plugins/Editor/AssetSearch/Source/Public/IAssetSearchModule.h b/Engine/Plugins/Editor/AssetSearch/Source/Public/IAssetSearchModule.h
--- a/Engine/Plugins/Editor/AssetSearch/Source/Public/IAssetSearchModule.h
+++ b/Engine/Plugins/Editor/AssetSearch/Source/Public/IAssetSearchModule.h
@@ -71,7 +71,7 @@ public:
 
 	virtual FSearchStats GetStats() const = 0;
 
-	virtual bool Search(const FSearchQuery& Query, TFunctionRef<bool(FSearchRecord&&)> InCallback) = 0;
+	virtual void Search(const FSearchQuery& Query, TFunction<void(TArray<FSearchRecord>&&)> InCallback) = 0;
 
 	virtual void RegisterIndexer(FName AssetClassName, IAssetIndexer* Indexer) = 0;
 	virtual void UnregisterIndexer(IAssetIndexer* Indexer) = 0;
