#UE4 No longer process objects for AssetSearchManager at save time during cook as it is too much to process quickly

diff --git a/Engine/Plugins/Editor/AssetSearch/Source/Private/AssetSearchManager.cpp b/Engine/Plugins/Editor/AssetSearch/Source/Private/AssetSearchManager.cpp
--- a/Engine/Plugins/Editor/AssetSearch/Source/Private/AssetSearchManager.cpp
+++ b/Engine/Plugins/Editor/AssetSearch/Source/Private/AssetSearchManager.cpp
@@ -111,7 +111,10 @@ void FAssetSearchManager::OnAssetAdded(const FAssetData& InAssetData)
 void FAssetSearchManager::OnObjectSaved(UObject* InObject)
 {
 	check(IsInGameThread());
-	StoreIndexForAsset(InObject);
+	if (!GIsCookerLoadingPackage)
+	{
+		StoreIndexForAsset(InObject);
+	}
 }
 
 void FAssetSearchManager::OnAssetLoaded(UObject* InObject)
