Don't load AssetSearch in commandlets

diff --git a/Engine/Plugins/Editor/AssetSearch/AssetSearch.uplugin b/Engine/Plugins/Editor/AssetSearch/AssetSearch.uplugin
--- a/Engine/Plugins/Editor/AssetSearch/AssetSearch.uplugin
+++ b/Engine/Plugins/Editor/AssetSearch/AssetSearch.uplugin
@@ -17,7 +17,7 @@
 	"Modules": [
 		{
 			"Name": "AssetSearch",
-			"Type": "Editor",
+			"Type": "EditorNoCommandlet",
 			"LoadingPhase": "Default"
 		}
 	],
