Re-enabling universal search, moving operations off the taskgraph and instead using its own low priority thread for handling feed/update/search requests through multiple buckets to permit prirotization and interrupting the queue.  This should prevent the search module from hogging the task graph and preventing GC.

diff --git a/Engine/Plugins/Editor/AssetSearch/Source/Private/AssetSearchManager.cpp b/Engine/Plugins/Editor/AssetSearch/Source/Private/AssetSearchManager.cpp
--- a/Engine/Plugins/Editor/AssetSearch/Source/Private/AssetSearchManager.cpp
+++ b/Engine/Plugins/Editor/AssetSearch/Source/Private/AssetSearchManager.cpp
@@ -49,10 +49,14 @@ FAssetSearchManager::FAssetSearchManager()
 	PendingDownloads = 0;
 	TotalSearchRecords = 0;
 	LastRecordCountUpdateSeconds = 0;
+
+	RunThread = false;
 }
 
 FAssetSearchManager::~FAssetSearchManager()
 {
+	RunThread = false;
+
 	{
 		FScopeLock ScopedLock(&SearchDatabaseCS);
 		FCoreUObjectDelegates::OnObjectSaved.RemoveAll(this);
@@ -84,6 +88,9 @@ void FAssetSearchManager::Start()
 	AssetRegistry.GetAllAssets(ProcessAssetQueue, true);
 
 	TickerHandle = FTicker::GetCoreTicker().AddTicker(FTickerDelegate::CreateRaw(this, &FAssetSearchManager::Tick_GameThread), 0);
+
+	RunThread = true;
+	DatabaseThread = FRunnableThread::Create(this, TEXT("UniversalSearch"), 0, TPri_BelowNormal);
 }
 
 FSearchStats FAssetSearchManager::GetStats() const
@@ -111,6 +118,7 @@ void FAssetSearchManager::OnAssetAdded(const FAssetData& InAssetData)
 void FAssetSearchManager::OnObjectSaved(UObject* InObject)
 {
 	check(IsInGameThread());
+
 	if (!GIsCookerLoadingPackage)
 	{
 		StoreIndexForAsset(InObject);
@@ -120,6 +128,7 @@ void FAssetSearchManager::OnObjectSaved(UObject* InObject)
 void FAssetSearchManager::OnAssetLoaded(UObject* InObject)
 {
 	check(IsInGameThread());
+
 	if (bIndexUnindexAssetsOnLoad)
 	{
 		RequestIndexAsset(InObject);
@@ -156,7 +165,8 @@ void FAssetSearchManager::RequestIndexAsset(UObject* InAsset)
 
 	FAssetData AssetData(InAsset);
 	FString AssetJsonDDCKey = TryGetDDCKeyForAsset(InAsset);
-	AsyncTask(ENamedThreads::AnyBackgroundThreadNormalTask, [this, AssetData, AssetWeakPtr, AssetJsonDDCKey]() {
+
+	UpdateOperations.Enqueue([this, AssetData, AssetWeakPtr, AssetJsonDDCKey]() {
 		FScopeLock ScopedLock(&SearchDatabaseCS);
 		if (!SearchDatabase.IsAssetUpToDate(AssetData, AssetJsonDDCKey))
 		{
@@ -222,7 +232,7 @@ bool FAssetSearchManager::TryLoadIndexForAsset(const FAssetData& InAsset)
 	if (!AssetJsonDDCKey.IsEmpty())
 	{
 		PendingDownloads++;
-		AsyncTask(ENamedThreads::AnyBackgroundThreadNormalTask, [this, InAsset, AssetJsonDDCKey]() {
+		FeedOperations.Enqueue([this, InAsset, AssetJsonDDCKey]() {
 			FScopeLock ScopedLock(&SearchDatabaseCS);
 			if (!SearchDatabase.IsAssetUpToDate(InAsset, AssetJsonDDCKey))
 			{
@@ -337,7 +347,7 @@ void FAssetSearchManager::AddOrUpdateAsset(const FAssetData& InAssetData, const
 	check(IsInGameThread());
 
 	PendingDatabaseUpdates++;
-	AsyncTask(ENamedThreads::AnyBackgroundThreadNormalTask, [this, InAssetData, IndexedJson, DerivedDataKey]() {
+	UpdateOperations.Enqueue([this, InAssetData, IndexedJson, DerivedDataKey]() {
 		FScopeLock ScopedLock(&SearchDatabaseCS);
 		SearchDatabase.AddOrUpdateAsset(InAssetData, IndexedJson, DerivedDataKey);
 		PendingDatabaseUpdates--;
@@ -387,7 +397,7 @@ bool FAssetSearchManager::Tick_GameThread(float DeltaTime)
 	{
 		LastRecordCountUpdateSeconds = FPlatformTime::Seconds();
 
-		AsyncTask(ENamedThreads::AnyBackgroundThreadNormalTask, [this]() {
+		ImmediateOperations.Enqueue([this]() {
 			FScopeLock ScopedLock(&SearchDatabaseCS);
 			TotalSearchRecords = SearchDatabase.GetTotalSearchRecords();
 		});
@@ -396,11 +406,33 @@ bool FAssetSearchManager::Tick_GameThread(float DeltaTime)
 	return true;
 }
 
+uint32 FAssetSearchManager::Run()
+{
+	Tick_DatabaseOperationThread();
+	return 0;
+}
+
+void FAssetSearchManager::Tick_DatabaseOperationThread()
+{
+	while (RunThread)
+	{
+		TFunction<void()> Operation;
+		if (ImmediateOperations.Dequeue(Operation) || FeedOperations.Dequeue(Operation) || UpdateOperations.Dequeue(Operation))
+		{
+			Operation();
+		}
+		else
+		{
+			FPlatformProcess::Sleep(0.1);
+		}
+	}
+}
+
 void FAssetSearchManager::Search(const FSearchQuery& Query, TFunction<void(TArray<FSearchRecord>&&)> InCallback)
 {
 	check(IsInGameThread());
 
-	AsyncTask(ENamedThreads::AnyBackgroundThreadNormalTask, [this, Query, InCallback]() {
+	ImmediateOperations.Enqueue([this, Query, InCallback]() {
 
 		TArray<FSearchRecord> Results;
 
diff --git a/Engine/Plugins/Editor/AssetSearch/Source/Private/AssetSearchManager.h b/Engine/Plugins/Editor/AssetSearch/Source/Private/AssetSearchManager.h
--- a/Engine/Plugins/Editor/AssetSearch/Source/Private/AssetSearchManager.h
+++ b/Engine/Plugins/Editor/AssetSearch/Source/Private/AssetSearchManager.h
@@ -9,7 +9,7 @@
 
 class FRunnableThread;
 
-class FAssetSearchManager
+class FAssetSearchManager : public FRunnable
 {
 public:
 	FAssetSearchManager();
@@ -24,6 +24,8 @@ public:
 
 private:
 	bool Tick_GameThread(float DeltaTime);
+	virtual uint32 Run() override;
+	void Tick_DatabaseOperationThread();
 
 private:
 	void OnAssetAdded(const FAssetData& InAssetData);
@@ -79,5 +81,10 @@ private:
 	FDelegateHandle TickerHandle;
 
 private:
-	FRunnableThread* Thread;
+	TAtomic<bool> RunThread;
+	FRunnableThread* DatabaseThread = nullptr;
+
+	TQueue<TFunction<void()>> ImmediateOperations;
+	TQueue<TFunction<void()>> FeedOperations;
+	TQueue<TFunction<void()>> UpdateOperations;
 };
\ No newline at end of file
diff --git a/Engine/Plugins/Editor/AssetSearch/Source/Private/Indexers/BlueprintIndexer.cpp b/Engine/Plugins/Editor/AssetSearch/Source/Private/Indexers/BlueprintIndexer.cpp
--- a/Engine/Plugins/Editor/AssetSearch/Source/Private/Indexers/BlueprintIndexer.cpp
+++ b/Engine/Plugins/Editor/AssetSearch/Source/Private/Indexers/BlueprintIndexer.cpp
@@ -7,6 +7,7 @@
 #include "K2Node_Variable.h"
 #include "Kismet2/BlueprintEditorUtils.h"
 #include "Utility/IndexerUtilities.h"
+#include "K2Node_BaseMCDelegate.h"
 
 #define LOCTEXT_NAMESPACE "FBlueprintIndexer"
 
@@ -58,24 +59,15 @@ void FBlueprintIndexer::IndexAsset(const UObject* InAssetObject, FSearchSerializ
 			
 			if (UK2Node_CallFunction* FunctionNode = Cast<UK2Node_CallFunction>(Node))
 			{
-				if (UFunction* TargetFunction = FunctionNode->GetTargetFunction())
-				{
-					Serializer.IndexProperty(TEXT("TargetFunction"), TargetFunction->GetName());
-				}
+				IndexMemberReference(Serializer, FunctionNode->FunctionReference, TEXT("Function"));
 			}
-			if (UK2Node_Variable* VariableNode = Cast<UK2Node_Variable>(Node))
+			else if (UK2Node_BaseMCDelegate* DelegateNode = Cast<UK2Node_BaseMCDelegate>(Node))
 			{
-				const FMemberReference& VariableReference = VariableNode->VariableReference;
-
-				Serializer.IndexProperty(TEXT("MemberName"), VariableReference.GetMemberName().ToString());
-				if (VariableReference.GetMemberGuid().IsValid())
-				{
-					Serializer.IndexProperty(TEXT("MemberGuid"), VariableReference.GetMemberGuid().ToString(EGuidFormats::Digits));
-				}
-				if (UClass* MemberParentClass = VariableReference.GetMemberParentClass())
-				{
-					Serializer.IndexProperty(TEXT("MemberParent"), MemberParentClass->GetPathName());
-				}
+				IndexMemberReference(Serializer, DelegateNode->DelegateReference, TEXT("Delegate"));
+			}
+			else if (UK2Node_Variable* VariableNode = Cast<UK2Node_Variable>(Node))
+			{
+				IndexMemberReference(Serializer, VariableNode->VariableReference, TEXT("Variable"));
 				//Serializer.WriteValue(TEXT("bSelfContext"), VariableReference.IsSelfContext());
 			}
 
@@ -98,6 +90,21 @@ void FBlueprintIndexer::IndexAsset(const UObject* InAssetObject, FSearchSerializ
 	}
 }
 
+void FBlueprintIndexer::IndexMemberReference(FSearchSerializer& Serializer, const FMemberReference& MemberReference, const FString& MemberType)
+{
+	Serializer.IndexProperty(MemberType + TEXT("Name"), MemberReference.GetMemberName());
+
+	if (MemberReference.GetMemberGuid().IsValid())
+	{
+		Serializer.IndexProperty(MemberType + TEXT("Guid"), MemberReference.GetMemberGuid().ToString(EGuidFormats::Digits));
+	}
+
+	if (UClass* MemberParentClass = MemberReference.GetMemberParentClass())
+	{
+		Serializer.IndexProperty(MemberType + TEXT("Parent"), MemberParentClass->GetPathName());
+	}
+}
+
 #undef LOCTEXT_NAMESPACE
 
 PRAGMA_ENABLE_OPTIMIZATION
\ No newline at end of file
diff --git a/Engine/Plugins/Editor/AssetSearch/Source/Private/Indexers/BlueprintIndexer.h b/Engine/Plugins/Editor/AssetSearch/Source/Private/Indexers/BlueprintIndexer.h
--- a/Engine/Plugins/Editor/AssetSearch/Source/Private/Indexers/BlueprintIndexer.h
+++ b/Engine/Plugins/Editor/AssetSearch/Source/Private/Indexers/BlueprintIndexer.h
@@ -5,9 +5,15 @@
 #include "CoreMinimal.h"
 #include "IAssetSearchModule.h"
 
+struct FMemberReference;
+
 class FBlueprintIndexer : public IAssetIndexer
 {
+public:
 	virtual FString GetName() const override { return TEXT("Blueprint"); }
 	virtual int32 GetVersion() const override;
 	virtual void IndexAsset(const UObject* InAssetObject, FSearchSerializer& Serializer) override;
+
+private:
+	void IndexMemberReference(FSearchSerializer& Serializer, const FMemberReference& MemberReference, const FString& MemberType);
 };
\ No newline at end of file
diff --git a/Engine/Plugins/Editor/AssetSearch/Source/Private/SearchSerializer.cpp b/Engine/Plugins/Editor/AssetSearch/Source/Private/SearchSerializer.cpp
--- a/Engine/Plugins/Editor/AssetSearch/Source/Private/SearchSerializer.cpp
+++ b/Engine/Plugins/Editor/AssetSearch/Source/Private/SearchSerializer.cpp
@@ -151,6 +151,11 @@ void FSearchSerializer::IndexProperty(const FString& InName, const FText& InValu
 	IndexProperty(*FTextProperty::StaticClass(), nullptr, InName, *FTextInspector::GetSourceString(InValue));
 }
 
+void FSearchSerializer::IndexProperty(const FString& InName, const FName& InValue)
+{
+	IndexProperty(*FNameProperty::StaticClass(), nullptr, InName, InValue.ToString());
+}
+
 void FSearchSerializer::IndexProperty(const FString& InName, const FString& InValue)
 {
 	IndexProperty(*FStrProperty::StaticClass(), nullptr, InName, InValue);
diff --git a/Engine/Plugins/Editor/AssetSearch/Source/Public/SearchSerializer.h b/Engine/Plugins/Editor/AssetSearch/Source/Public/SearchSerializer.h
--- a/Engine/Plugins/Editor/AssetSearch/Source/Public/SearchSerializer.h
+++ b/Engine/Plugins/Editor/AssetSearch/Source/Public/SearchSerializer.h
@@ -31,6 +31,7 @@ public:
 
 	//---------------------------------------------------------------
 
+	ASSETSEARCH_API void IndexProperty(const FString& InName, const FName& InValue);
 	ASSETSEARCH_API void IndexProperty(const FString& InName, const FText& InValue);
 	ASSETSEARCH_API void IndexProperty(const FString& InName, const FString& InValue);
 	ASSETSEARCH_API void IndexProperty(const UClass* InPropertyClass, const FString& InName, const FString& InValue);
