Fixing some IWYU errors.

diff --git a/Engine/Plugins/Editor/AssetSearch/Source/Private/AssetSearchManager.cpp b/Engine/Plugins/Editor/AssetSearch/Source/Private/AssetSearchManager.cpp
--- a/Engine/Plugins/Editor/AssetSearch/Source/Private/AssetSearchManager.cpp
+++ b/Engine/Plugins/Editor/AssetSearch/Source/Private/AssetSearchManager.cpp
@@ -14,6 +14,7 @@
 #include "Containers/StringConv.h"
 #include "Containers/Ticker.h"
 #include "Misc/Paths.h"
+#include "HAL/RunnableThread.h"
 
 static bool bIndexUnindexAssetsOnLoad = false;
 FAutoConsoleVariableRef CVarIndexUnindexAssetsOnLoad(
diff --git a/Engine/Plugins/Editor/AssetSearch/Source/Private/AssetSearchManager.h b/Engine/Plugins/Editor/AssetSearch/Source/Private/AssetSearchManager.h
--- a/Engine/Plugins/Editor/AssetSearch/Source/Private/AssetSearchManager.h
+++ b/Engine/Plugins/Editor/AssetSearch/Source/Private/AssetSearchManager.h
@@ -6,6 +6,7 @@
 #include "IAssetSearchModule.h"
 #include "AssetSearchDatabase.h"
 #include "Containers/Queue.h"
+#include "HAL/Runnable.h"
 
 class FRunnableThread;
 
