Add studio analytic event for using asset search.

diff --git a/Engine/Plugins/Editor/AssetSearch/Source/Private/AssetSearchManager.cpp b/Engine/Plugins/Editor/AssetSearch/Source/Private/AssetSearchManager.cpp
--- a/Engine/Plugins/Editor/AssetSearch/Source/Private/AssetSearchManager.cpp
+++ b/Engine/Plugins/Editor/AssetSearch/Source/Private/AssetSearchManager.cpp
@@ -15,6 +15,8 @@
 #include "Containers/Ticker.h"
 #include "Misc/Paths.h"
 #include "HAL/RunnableThread.h"
+#include "StudioAnalytics.h"
+#include "AnalyticsEventAttribute.h"
 
 static bool bIndexUnindexAssetsOnLoad = false;
 FAutoConsoleVariableRef CVarIndexUnindexAssetsOnLoad(
@@ -433,6 +435,10 @@ void FAssetSearchManager::Search(const FSearchQuery& Query, TFunction<void(TArra
 {
 	check(IsInGameThread());
 
+	FStudioAnalytics::ReportEvent(TEXT("AssetSearch"), {
+		FAnalyticsEventAttribute(TEXT("QueryString"), Query.Query)
+	});
+
 	ImmediateOperations.Enqueue([this, Query, InCallback]() {
 
 		TArray<FSearchRecord> Results;
