Fixing a crash on exit waiting for the database operation thread to finish.  Ignore assets in developer directories other than yours.  Updating the database to have unique flag on asset_path, and including an index to speed up deletes for properties by indexing assetids.  Adding a way to manually index missing assets.

diff --git a/Engine/Plugins/Editor/AssetSearch/Source/Private/AssetSearchDatabase.cpp b/Engine/Plugins/Editor/AssetSearch/Source/Private/AssetSearchDatabase.cpp
--- a/Engine/Plugins/Editor/AssetSearch/Source/Private/AssetSearchDatabase.cpp
+++ b/Engine/Plugins/Editor/AssetSearch/Source/Private/AssetSearchDatabase.cpp
@@ -19,8 +19,11 @@ enum class EAssetSearchDatabaseVersion
 {
 	Empty = 0,
 	Initial = 1,
+	IndexingAssetIdsAssetPathsUnique = 2,
 
-	Current = Initial,
+	// -----<new versions can be added above this line>-------------------------------------------------
+	VersionPlusOne,
+	LatestVersion = VersionPlusOne - 1
 };
 
 class FAssetSearchDatabaseStatements
@@ -163,7 +166,8 @@ public:
 		if (Statement_AddAssetToAssetTable.BindAndExecute(InAssetData.AssetName.ToString(), InAssetData.AssetClass.ToString(), InAssetData.ObjectPath.ToString(), IndexedJsonHash))
 		{
 			int64 AssetId = Database.GetLastInsertRowId();
-			//ensure(AssetId == GetAssetIdForAsset(InAssetData));
+
+			BeginTransaction();
 
 			TSharedRef<TJsonReader<>> JsonReader = TJsonReaderFactory<>::Create(IndexedJson);
 
@@ -258,8 +262,11 @@ public:
 				}
 			}
 
+			CommitTransaction();
+
 			return true;
 		}
+
 		return false;
 	}
 
@@ -473,31 +480,37 @@ bool FAssetSearchDatabase::Open(const FString& InSessionPath, const ESQLiteDatab
 
 	// Set the database to use exclusive WAL mode for performance (exclusive works even on platforms without a mmap implementation)
 	// Set the database "NORMAL" fsync mode to only perform a fsync when check-pointing the WAL to the main database file (fewer fsync calls are better for performance, with a very slight loss of WAL durability if the power fails)
+	Database->Execute(TEXT("PRAGMA cache_size=1000;"));
+	Database->Execute(TEXT("PRAGMA page_size=65535;"));
 	Database->Execute(TEXT("PRAGMA locking_mode=EXCLUSIVE;"));
-	Database->Execute(TEXT("PRAGMA journal_mode=WAL;"));
-	Database->Execute(TEXT("PRAGMA synchronous=NORMAL;"));
+
+	//Database->Execute(TEXT("PRAGMA journal_mode=WAL;"));
+	//Database->Execute(TEXT("PRAGMA synchronous=NORMAL;"));
+
+	Database->Execute(TEXT("PRAGMA journal_mode=NORMAL;"));
+	Database->Execute(TEXT("PRAGMA synchronous=OFF;"));
 
 	int32 LoadedDatabaseVersion = 0;
 	Database->GetUserVersion(LoadedDatabaseVersion);
 	if (LoadedDatabaseVersion != (int32)EAssetSearchDatabaseVersion::Empty)
 	{
-		if (LoadedDatabaseVersion > (int32)EAssetSearchDatabaseVersion::Current)
+		if (LoadedDatabaseVersion > (int32)EAssetSearchDatabaseVersion::LatestVersion)
 		{
 			Close();
-			UE_LOG(LogAssetSearch, Error, TEXT("Failed to open database for '%s': Database is too new (version %d, expected = %d)"), *InSessionPath, LoadedDatabaseVersion, (int32)EAssetSearchDatabaseVersion::Current);
+			UE_LOG(LogAssetSearch, Error, TEXT("Failed to open database for '%s': Database is too new (version %d, expected = %d)"), *InSessionPath, LoadedDatabaseVersion, (int32)EAssetSearchDatabaseVersion::LatestVersion);
 			return false;
 		}
-		else if (LoadedDatabaseVersion < (int32)EAssetSearchDatabaseVersion::Current)
+		else if (LoadedDatabaseVersion < (int32)EAssetSearchDatabaseVersion::LatestVersion)
 		{
 			Close(true);
-			UE_LOG(LogAssetSearch, Log, TEXT("Opened database '%s': Database is too old (version %d, expected = %d), creating new database"), *InSessionPath, LoadedDatabaseVersion, (int32)EAssetSearchDatabaseVersion::Current);
+			UE_LOG(LogAssetSearch, Log, TEXT("Opened database '%s': Database is too old (version %d, expected = %d), creating new database"), *InSessionPath, LoadedDatabaseVersion, (int32)EAssetSearchDatabaseVersion::LatestVersion);
 			return Open(InSessionPath, InOpenMode);
 		}
 	}
 
 	// Create our required tables
 	//========================================================================
-	if (!ensure(Database->Execute(TEXT("CREATE TABLE IF NOT EXISTS table_assets(assetid INTEGER PRIMARY KEY, asset_name, asset_class, asset_path, index_hash);"))))
+	if (!ensure(Database->Execute(TEXT("CREATE TABLE IF NOT EXISTS table_assets(assetid INTEGER PRIMARY KEY, asset_name, asset_class, asset_path TEXT UNIQUE, index_hash);"))))
 	{
 		LogLastError();
 		Close();
@@ -610,10 +623,19 @@ bool FAssetSearchDatabase::Open(const FString& InSessionPath, const ESQLiteDatab
 		return false;
 	}
 
+	if (!ensure(Database->Execute(
+		TEXT("CREATE INDEX IF NOT EXISTS assetid_index ON table_asset_properties(assetid);")
+	)))
+	{
+		LogLastError();
+		Close();
+		return false;
+	}
+
 	//CREATE INDEX blueprint_nodes ON table_assets(json_extract(json, '$.Blueprint.Nodes'))
 
 	// The database will have the latest schema at this point, so update the user-version
-	if (!Database->SetUserVersion((int32)EAssetSearchDatabaseVersion::Current))
+	if (!Database->SetUserVersion((int32)EAssetSearchDatabaseVersion::LatestVersion))
 	{
 		Close();
 		return false;
@@ -719,4 +741,20 @@ int64 FAssetSearchDatabase::GetTotalSearchRecords() const
 	return INDEX_NONE;
 }
 
+void FAssetSearchDatabase::BeginTransaction()
+{
+	if (ensure(Statements))
+	{
+		Statements->BeginTransaction();
+	}
+}
+
+void FAssetSearchDatabase::CommitTransaction()
+{
+	if (ensure(Statements))
+	{
+		Statements->CommitTransaction();
+	}
+}
+
 PRAGMA_ENABLE_OPTIMIZATION
diff --git a/Engine/Plugins/Editor/AssetSearch/Source/Private/AssetSearchDatabase.h b/Engine/Plugins/Editor/AssetSearch/Source/Private/AssetSearchDatabase.h
--- a/Engine/Plugins/Editor/AssetSearch/Source/Private/AssetSearchDatabase.h
+++ b/Engine/Plugins/Editor/AssetSearch/Source/Private/AssetSearchDatabase.h
@@ -84,6 +84,9 @@ public:
 	 */
 	int64 GetTotalSearchRecords() const;
 
+	void BeginTransaction();
+	void CommitTransaction();
+
 private:
 	void LogLastError() const;
 
diff --git a/Engine/Plugins/Editor/AssetSearch/Source/Private/AssetSearchManager.cpp b/Engine/Plugins/Editor/AssetSearch/Source/Private/AssetSearchManager.cpp
--- a/Engine/Plugins/Editor/AssetSearch/Source/Private/AssetSearchManager.cpp
+++ b/Engine/Plugins/Editor/AssetSearch/Source/Private/AssetSearchManager.cpp
@@ -17,6 +17,9 @@
 #include "HAL/RunnableThread.h"
 #include "StudioAnalytics.h"
 #include "AnalyticsEventAttribute.h"
+#include "Misc/FeedbackContext.h"
+
+#define LOCTEXT_NAMESPACE "FAssetSearchManager"
 
 static bool bIndexUnindexAssetsOnLoad = false;
 FAutoConsoleVariableRef CVarIndexUnindexAssetsOnLoad(
@@ -59,6 +62,7 @@ FAssetSearchManager::FAssetSearchManager()
 FAssetSearchManager::~FAssetSearchManager()
 {
 	RunThread = false;
+	DatabaseThread->WaitForCompletion();
 
 	{
 		FScopeLock ScopedLock(&SearchDatabaseCS);
@@ -88,7 +92,14 @@ void FAssetSearchManager::Start()
 
 	IAssetRegistry& AssetRegistry = FModuleManager::LoadModuleChecked<FAssetRegistryModule>("AssetRegistry").Get();
 	AssetRegistry.OnAssetAdded().AddRaw(this, &FAssetSearchManager::OnAssetAdded);
-	AssetRegistry.GetAllAssets(ProcessAssetQueue, true);
+	
+	TArray<FAssetData> TempAssetData;
+	AssetRegistry.GetAllAssets(TempAssetData, true);
+
+	for (const FAssetData& Data : TempAssetData)
+	{
+		OnAssetAdded(Data);
+	}
 
 	TickerHandle = FTicker::GetCoreTicker().AddTicker(FTickerDelegate::CreateRaw(this, &FAssetSearchManager::Tick_GameThread), 0);
 
@@ -103,6 +114,7 @@ FSearchStats FAssetSearchManager::GetStats() const
 	Stats.Downloading = PendingDownloads;
 	Stats.PendingDatabaseUpdates = PendingDatabaseUpdates;
 	Stats.TotalRecords = TotalSearchRecords;
+	Stats.AssetsMissingIndex = FailedDDCRequests.Num();
 	return Stats;
 }
 
@@ -115,6 +127,20 @@ void FAssetSearchManager::RegisterIndexer(FName AssetClassName, IAssetIndexer* I
 void FAssetSearchManager::OnAssetAdded(const FAssetData& InAssetData)
 {
 	check(IsInGameThread());
+
+	static const FString DeveloperPathWithSlash = FPackageName::FilenameToLongPackageName(FPaths::GameDevelopersDir());
+	static const FString UsersDeveloperPathWithSlash = FPackageName::FilenameToLongPackageName(FPaths::GameUserDeveloperDir());
+	
+	// Don't process stuff in the other developer folders.
+	FString PackageName = InAssetData.PackageName.ToString();
+	if (PackageName.StartsWith(DeveloperPathWithSlash))
+	{
+		if (!PackageName.StartsWith(UsersDeveloperPathWithSlash))
+		{
+			return;
+		}
+	}
+
 	ProcessAssetQueue.Add(InAssetData);
 }
 
@@ -124,7 +150,7 @@ void FAssetSearchManager::OnObjectSaved(UObject* InObject)
 
 	if (!GIsCookerLoadingPackage)
 	{
-		StoreIndexForAsset(InObject);
+		StoreIndexForAsset(InObject, false);
 	}
 }
 
@@ -234,21 +260,18 @@ bool FAssetSearchManager::TryLoadIndexForAsset(const FAssetData& InAsset)
 
 	if (!AssetJsonDDCKey.IsEmpty())
 	{
-		PendingDownloads++;
 		FeedOperations.Enqueue([this, InAsset, AssetJsonDDCKey]() {
 			FScopeLock ScopedLock(&SearchDatabaseCS);
 			if (!SearchDatabase.IsAssetUpToDate(InAsset, AssetJsonDDCKey))
 			{
+				PendingDownloads++;
+
 				FAssetDDCRequest DDCRequest;
 				DDCRequest.AssetData = InAsset;
 				DDCRequest.DDCKey_IndexDataHash = AssetJsonDDCKey;
 				DDCRequest.DDCHandle = GetDerivedDataCacheRef().GetAsynchronous(*AssetJsonDDCKey, InAsset.ObjectPath.ToString());
 				ProcessDDCQueue.Enqueue(DDCRequest);
 			}
-			else
-			{
-				PendingDownloads--;
-			}
 		});
 
 		return true;
@@ -387,6 +410,10 @@ bool FAssetSearchManager::Tick_GameThread(float DeltaTime)
 			{
 				LoadDDCContentIntoDatabase(PendingRequest->AssetData, OutContent, PendingRequest->DDCKey_IndexDataHash);
 			}
+			else
+			{
+				FailedDDCRequests.Add(*PendingRequest);
+			}
 
 			ProcessDDCQueue.Pop();
 			PendingDownloads--;
@@ -431,6 +458,34 @@ void FAssetSearchManager::Tick_DatabaseOperationThread()
 	}
 }
 
+void FAssetSearchManager::ForceIndexOnAssetsMissingIndex()
+{
+	check(IsInGameThread());
+
+	GWarn->BeginSlowTask(LOCTEXT("ForceIndexOnAssetsMissingIndex", "Indexing Assets"), true, true);
+
+	int32 RequestCount = 0;
+	for (const FAssetDDCRequest& Request : FailedDDCRequests)
+	{
+		if (GWarn->ReceivedUserCancel())
+		{
+			break;
+		}
+
+		if (UObject* AssetToIndex = Request.AssetData.GetAsset())
+		{
+			StoreIndexForAsset(AssetToIndex, true);
+		}
+
+		RequestCount++;
+		GWarn->StatusForceUpdate(RequestCount, FailedDDCRequests.Num(), FText::FromString(Request.AssetData.PackageName.ToString()));
+	}
+
+	GWarn->EndSlowTask();
+
+	FailedDDCRequests.Reset();
+}
+
 void FAssetSearchManager::Search(const FSearchQuery& Query, TFunction<void(TArray<FSearchRecord>&&)> InCallback)
 {
 	check(IsInGameThread());
@@ -455,4 +510,6 @@ void FAssetSearchManager::Search(const FSearchQuery& Query, TFunction<void(TArra
 			InCallback(MoveTemp(ResultsFwd));
 		});
 	});
-}
\ No newline at end of file
+}
+
+#undef LOCTEXT_NAMESPACE
\ No newline at end of file
diff --git a/Engine/Plugins/Editor/AssetSearch/Source/Private/AssetSearchManager.h b/Engine/Plugins/Editor/AssetSearch/Source/Private/AssetSearchManager.h
--- a/Engine/Plugins/Editor/AssetSearch/Source/Private/AssetSearchManager.h
+++ b/Engine/Plugins/Editor/AssetSearch/Source/Private/AssetSearchManager.h
@@ -23,6 +23,9 @@ public:
 
 	void Search(const FSearchQuery& Query, TFunction<void(TArray<FSearchRecord>&&)> InCallback);
 
+	// Utility
+	void ForceIndexOnAssetsMissingIndex();
+
 private:
 	bool Tick_GameThread(float DeltaTime);
 	virtual uint32 Run() override;
@@ -44,7 +47,7 @@ private:
 	FString GetDerivedDataKey(const FSHAHash& IndexedContentHash);
 	FString GetDerivedDataKey(const FAssetData& UnindexedAsset);
 	bool HasIndexerForClass(UClass* AssetClass);
-	void StoreIndexForAsset(UObject* InAsset, bool Unindexed = false);
+	void StoreIndexForAsset(UObject* InAsset, bool Unindexed);
 	void LoadDDCContentIntoDatabase(const FAssetData& InAsset, const TArray<uint8>& Content, const FString& DerivedDataKey);
 
 private:
@@ -79,6 +82,8 @@ private:
 	};
 	TQueue<FAssetDDCRequest, EQueueMode::Mpsc> ProcessDDCQueue;
 
+	TArray<FAssetDDCRequest> FailedDDCRequests;
+
 	FDelegateHandle TickerHandle;
 
 private:
diff --git a/Engine/Plugins/Editor/AssetSearch/Source/Private/AssetSearchModule.cpp b/Engine/Plugins/Editor/AssetSearch/Source/Private/AssetSearchModule.cpp
--- a/Engine/Plugins/Editor/AssetSearch/Source/Private/AssetSearchModule.cpp
+++ b/Engine/Plugins/Editor/AssetSearch/Source/Private/AssetSearchModule.cpp
@@ -71,6 +71,11 @@ public:
 		SearchManager->Search(Query, InCallback);
 	}
 
+	virtual void ForceIndexOnAssetsMissingIndex() override
+	{
+		SearchManager->ForceIndexOnAssetsMissingIndex();
+	}
+
 	virtual void RegisterIndexer(FName AssetClassName, IAssetIndexer* Indexer) override
 	{
 		SearchManager->RegisterIndexer(AssetClassName, Indexer);
diff --git a/Engine/Plugins/Editor/AssetSearch/Source/Private/Widgets/SSearchBrowser.cpp b/Engine/Plugins/Editor/AssetSearch/Source/Private/Widgets/SSearchBrowser.cpp
--- a/Engine/Plugins/Editor/AssetSearch/Source/Private/Widgets/SSearchBrowser.cpp
+++ b/Engine/Plugins/Editor/AssetSearch/Source/Private/Widgets/SSearchBrowser.cpp
@@ -31,6 +31,7 @@
 #include "AssetRegistryModule.h"
 #include "Widgets/Views/STreeView.h"
 #include "Widgets/Layout/SSeparator.h"
+#include "Widgets/Input/SHyperlink.h"
 
 #define LOCTEXT_NAMESPACE "SObjectBrowser"
 
@@ -164,6 +165,15 @@ void SSearchBrowser::Construct( const FArguments& InArgs )
 					SNew(STextBlock)
 					.Text(this, &SSearchBrowser::GetStatusText)
 				]
+
+				// Index unindexed items
+				+ SHorizontalBox::Slot()
+				.AutoWidth()
+				[
+					SNew(SHyperlink)
+					.Text(this, &SSearchBrowser::GetUnindexedAssetsText)
+					.OnNavigate(this, &SSearchBrowser::HandleForceIndexOfAssetsMissingIndex)
+				]
 			]
 		]
 	];
@@ -178,6 +188,19 @@ FText SSearchBrowser::GetStatusText() const
 	return FText::Format(LOCTEXT("SearchStatusTextFmt", "Scanning {0}   Downloading {1}   Updating {2}            Total Records {3}"), SearchStats.Scanning, SearchStats.Downloading, SearchStats.PendingDatabaseUpdates, SearchStats.TotalRecords);
 }
 
+FText SSearchBrowser::GetUnindexedAssetsText() const
+{
+	IAssetSearchModule& SearchModule = IAssetSearchModule::Get();
+	FSearchStats SearchStats = SearchModule.GetStats();
+	return FText::Format(LOCTEXT("UnindexedAssetsLinkFormat", "{0} Missing"), SearchStats.AssetsMissingIndex);
+}
+
+void SSearchBrowser::HandleForceIndexOfAssetsMissingIndex()
+{
+	IAssetSearchModule& SearchModule = IAssetSearchModule::Get();
+	SearchModule.ForceIndexOnAssetsMissingIndex();
+}
+
 FReply SSearchBrowser::OnRefresh()
 {
 	RefreshList();
diff --git a/Engine/Plugins/Editor/AssetSearch/Source/Private/Widgets/SSearchBrowser.h b/Engine/Plugins/Editor/AssetSearch/Source/Private/Widgets/SSearchBrowser.h
--- a/Engine/Plugins/Editor/AssetSearch/Source/Private/Widgets/SSearchBrowser.h
+++ b/Engine/Plugins/Editor/AssetSearch/Source/Private/Widgets/SSearchBrowser.h
@@ -41,6 +41,9 @@ public:
 private:
 
 	FText GetStatusText() const;
+	FText GetUnindexedAssetsText() const;
+
+	void HandleForceIndexOfAssetsMissingIndex();
 
 	FReply OnRefresh();
 
diff --git a/Engine/Plugins/Editor/AssetSearch/Source/Public/IAssetSearchModule.h b/Engine/Plugins/Editor/AssetSearch/Source/Public/IAssetSearchModule.h
--- a/Engine/Plugins/Editor/AssetSearch/Source/Public/IAssetSearchModule.h
+++ b/Engine/Plugins/Editor/AssetSearch/Source/Public/IAssetSearchModule.h
@@ -40,6 +40,8 @@ struct FSearchStats
 	int32 Downloading = 0;
 	int32 PendingDatabaseUpdates = 0;
 
+	int32 AssetsMissingIndex = 0;
+
 	int64 TotalRecords = 0;
 };
 
@@ -73,6 +75,8 @@ public:
 
 	virtual void Search(const FSearchQuery& Query, TFunction<void(TArray<FSearchRecord>&&)> InCallback) = 0;
 
+	virtual void ForceIndexOnAssetsMissingIndex() = 0;
+
 	virtual void RegisterIndexer(FName AssetClassName, IAssetIndexer* Indexer) = 0;
 	virtual void UnregisterIndexer(IAssetIndexer* Indexer) = 0;
 
