Universal Search - Now monitoring deletes that happen both while the editor is running, and deletes that happen outside the editor during editor boot after initial scan completes.

diff --git a/Engine/Plugins/Editor/AssetSearch/Source/Private/AssetSearchDatabase.cpp b/Engine/Plugins/Editor/AssetSearch/Source/Private/AssetSearchDatabase.cpp
--- a/Engine/Plugins/Editor/AssetSearch/Source/Private/AssetSearchDatabase.cpp
+++ b/Engine/Plugins/Editor/AssetSearch/Source/Private/AssetSearchDatabase.cpp
@@ -50,6 +50,7 @@ public:
 		PREPARE_STATEMENT(Statement_CommitTransaction);
 		PREPARE_STATEMENT(Statement_RollbackTransaction);
 
+		PREPARE_STATEMENT(Statement_GetAllAssets);
 		PREPARE_STATEMENT(Statement_GetAssetIdForAssetPath);
 		PREPARE_STATEMENT(Statement_IsAssetUpToDate);
 		PREPARE_STATEMENT(Statement_GetTotalSearchRecords);
@@ -274,7 +275,11 @@ public:
 	private: FDeleteEntriesForAsset Statement_DeleteEntriesForAsset;
 	public: bool DeleteEntriesForAsset(const FAssetData& InAssetData)
 	{
-		if (Statement_DeleteEntriesForAsset.BindAndExecute(InAssetData.ObjectPath.ToString()))
+		return DeleteEntriesForAsset(InAssetData.ObjectPath.ToString());
+	}
+	public: bool DeleteEntriesForAsset(const FString& InAssetObjectPath)
+	{
+		if (Statement_DeleteEntriesForAsset.BindAndExecute(InAssetObjectPath))
 		{
 			return true;
 		}
@@ -391,6 +396,25 @@ public:
 			return ESQLitePreparedStatementExecuteRowResult::Error;
 		}) != INDEX_NONE;
 	}
+
+	SQLITE_PREPARED_STATEMENT(FGetAllAssets,
+		"SELECT asset_path FROM table_assets;",
+		SQLITE_PREPARED_STATEMENT_COLUMNS(FString),
+		SQLITE_PREPARED_STATEMENT_BINDINGS()
+	);
+	FGetAllAssets Statement_GetAllAssets;
+	bool GetAllAssets(TFunctionRef<ESQLitePreparedStatementExecuteRowResult(FString&&)> InCallback)
+	{
+		return Statement_GetAllAssets.BindAndExecute([&InCallback](const FGetAllAssets& InStatement)
+		{
+			FString AssetPath;
+			if (InStatement.GetColumnValues(AssetPath))
+			{
+				return InCallback(MoveTemp(AssetPath));
+			}
+			return ESQLitePreparedStatementExecuteRowResult::Error;
+		}) != INDEX_NONE;
+	}
 };
 
 class FAssetSearchDatabaseScopedTransaction
@@ -741,19 +765,38 @@ int64 FAssetSearchDatabase::GetTotalSearchRecords() const
 	return INDEX_NONE;
 }
 
-void FAssetSearchDatabase::BeginTransaction()
+void FAssetSearchDatabase::RemoveAsset(const FAssetData& InAssetData)
 {
-	if (ensure(Statements))
+	if (!ensure(Statements->DeleteEntriesForAsset(InAssetData)))
 	{
-		Statements->BeginTransaction();
+		LogLastError();
 	}
 }
 
-void FAssetSearchDatabase::CommitTransaction()
+void FAssetSearchDatabase::RemoveAssetsNotInThisSet(const TArray<FAssetData>& InAssets)
 {
-	if (ensure(Statements))
+	TSet<FString> InAssetPaths;
+	for (const FAssetData& InAsset : InAssets)
+	{
+		InAssetPaths.Add(InAsset.ObjectPath.ToString());
+	}
+
+	TArray<FString> MissingAssets;
+
+	Statements->GetAllAssets([&MissingAssets, &InAssetPaths](FString&& InResult)
+	{
+		if (!InAssetPaths.Contains(InResult))
+		{
+			MissingAssets.Add(InResult);
+		}
+
+		return ESQLitePreparedStatementExecuteRowResult::Continue;
+	});
+
+
+	for (const FString& MissingAsset : MissingAssets)
 	{
-		Statements->CommitTransaction();
+		Statements->DeleteEntriesForAsset(MissingAsset);
 	}
 }
 
diff --git a/Engine/Plugins/Editor/AssetSearch/Source/Private/AssetSearchDatabase.h b/Engine/Plugins/Editor/AssetSearch/Source/Private/AssetSearchDatabase.h
--- a/Engine/Plugins/Editor/AssetSearch/Source/Private/AssetSearchDatabase.h
+++ b/Engine/Plugins/Editor/AssetSearch/Source/Private/AssetSearchDatabase.h
@@ -84,8 +84,15 @@ public:
 	 */
 	int64 GetTotalSearchRecords() const;
 
-	void BeginTransaction();
-	void CommitTransaction();
+	/**
+	 * Remove asset from the the search database.
+	 */
+	void RemoveAsset(const FAssetData& InAssetData);
+
+	/**
+	 * Remove asset from the the search database that are not in this set of assets.
+	 */
+	void RemoveAssetsNotInThisSet(const TArray<FAssetData>& InAssets);
 
 private:
 	void LogLastError() const;
diff --git a/Engine/Plugins/Editor/AssetSearch/Source/Private/AssetSearchManager.cpp b/Engine/Plugins/Editor/AssetSearch/Source/Private/AssetSearchManager.cpp
--- a/Engine/Plugins/Editor/AssetSearch/Source/Private/AssetSearchManager.cpp
+++ b/Engine/Plugins/Editor/AssetSearch/Source/Private/AssetSearchManager.cpp
@@ -92,6 +92,8 @@ void FAssetSearchManager::Start()
 
 	IAssetRegistry& AssetRegistry = FModuleManager::LoadModuleChecked<FAssetRegistryModule>("AssetRegistry").Get();
 	AssetRegistry.OnAssetAdded().AddRaw(this, &FAssetSearchManager::OnAssetAdded);
+	AssetRegistry.OnAssetRemoved().AddRaw(this, &FAssetSearchManager::OnAssetRemoved);
+	AssetRegistry.OnFilesLoaded().AddRaw(this, &FAssetSearchManager::OnAssetScanFinished);
 	
 	TArray<FAssetData> TempAssetData;
 	AssetRegistry.GetAllAssets(TempAssetData, true);
@@ -141,7 +143,33 @@ void FAssetSearchManager::OnAssetAdded(const FAssetData& InAssetData)
 		}
 	}
 
-	ProcessAssetQueue.Add(InAssetData);
+	FAssetOperation Operation;
+	Operation.Asset = InAssetData;
+	ProcessAssetQueue.Add(Operation);
+}
+
+void FAssetSearchManager::OnAssetRemoved(const FAssetData& InAssetData)
+{
+	check(IsInGameThread());
+
+	FAssetOperation Operation;
+	Operation.Asset = InAssetData;
+	Operation.bRemoval = true;
+	ProcessAssetQueue.Add(Operation);
+}
+
+void FAssetSearchManager::OnAssetScanFinished()
+{
+	check(IsInGameThread());
+
+	TArray<FAssetData> AllAssets;
+	IAssetRegistry& AssetRegistry = FModuleManager::LoadModuleChecked<FAssetRegistryModule>("AssetRegistry").Get();
+	AssetRegistry.GetAllAssets(AllAssets, false);
+
+	UpdateOperations.Enqueue([this, AssetsAvailable = MoveTemp(AllAssets)]() mutable {
+		FScopeLock ScopedLock(&SearchDatabaseCS);
+		SearchDatabase.RemoveAssetsNotInThisSet(AssetsAvailable);
+	});
 }
 
 void FAssetSearchManager::OnObjectSaved(UObject* InObject)
@@ -168,7 +196,7 @@ void FAssetSearchManager::OnGetAssetTags(const UObject* Object, TArray<UObject::
 {
 	check(IsInGameThread());
 
-	FString ObjectPath = Object->GetPathName();
+	const FString ObjectPath = Object->GetPathName();
 	const FContentHashEntry* ContentEntry = ContentHashCache.Find(ObjectPath);
 
 	if (ContentEntry)
@@ -387,10 +415,22 @@ bool FAssetSearchManager::Tick_GameThread(float DeltaTime)
 	int32 ScanLimit = GameThread_AssetScanLimit;
 	while (ProcessAssetQueue.Num() > 0 && ScanLimit > 0 && PendingDownloads < PendingDownloadsMax)
 	{
-		FAssetData Asset = ProcessAssetQueue.Pop(false);
-		if (TryLoadIndexForAsset(Asset))
+		FAssetOperation Operation = ProcessAssetQueue.Pop(false);
+		FAssetData Asset = Operation.Asset;
+
+		if (Operation.bRemoval)
+		{
+			UpdateOperations.Enqueue([this, Asset]() {
+				FScopeLock ScopedLock(&SearchDatabaseCS);
+				SearchDatabase.RemoveAsset(Asset);
+			});
+		}
+		else
 		{
-			ScanLimit -= 10;
+			if (TryLoadIndexForAsset(Asset))
+			{
+				ScanLimit -= 10;
+			}
 		}
 
 		ScanLimit--;
@@ -474,7 +514,7 @@ void FAssetSearchManager::ForceIndexOnAssetsMissingIndex()
 
 		if (UObject* AssetToIndex = Request.AssetData.GetAsset())
 		{
-			StoreIndexForAsset(AssetToIndex, true);
+			StoreIndexForAsset(AssetToIndex, Request.DDCKey_IndexDataHash.StartsWith(TEXT("AssetSearch_Legacy")));
 		}
 
 		RequestCount++;
diff --git a/Engine/Plugins/Editor/AssetSearch/Source/Private/AssetSearchManager.h b/Engine/Plugins/Editor/AssetSearch/Source/Private/AssetSearchManager.h
--- a/Engine/Plugins/Editor/AssetSearch/Source/Private/AssetSearchManager.h
+++ b/Engine/Plugins/Editor/AssetSearch/Source/Private/AssetSearchManager.h
@@ -33,6 +33,9 @@ private:
 
 private:
 	void OnAssetAdded(const FAssetData& InAssetData);
+	void OnAssetRemoved(const FAssetData& InAssetData);
+	void OnAssetScanFinished();
+
 	void OnObjectSaved(UObject* InObject);
 	void OnAssetLoaded(UObject* InObject);
 	void OnGetAssetTags(const UObject* Object, TArray<UObject::FAssetRegistryTag>& OutTags);
@@ -72,7 +75,13 @@ private:
 
 	TArray<TWeakObjectPtr<UObject>> RequestIndexQueue;
 
-	TArray<FAssetData> ProcessAssetQueue;
+	struct FAssetOperation
+	{
+		FAssetData Asset;
+		bool bRemoval = false;
+	};
+
+	TArray<FAssetOperation> ProcessAssetQueue;
 
 	struct FAssetDDCRequest
 	{
