Universal Search - Now indexing all DialogueWave assets for their SpokeText.

diff --git a/Engine/Plugins/Editor/AssetSearch/Source/Private/AssetSearchManager.cpp b/Engine/Plugins/Editor/AssetSearch/Source/Private/AssetSearchManager.cpp
--- a/Engine/Plugins/Editor/AssetSearch/Source/Private/AssetSearchManager.cpp
+++ b/Engine/Plugins/Editor/AssetSearch/Source/Private/AssetSearchManager.cpp
@@ -18,6 +18,12 @@
 #include "StudioAnalytics.h"
 #include "AnalyticsEventAttribute.h"
 #include "Misc/FeedbackContext.h"
+#include "WidgetBlueprint.h"
+#include "Engine/Blueprint.h"
+#include "Engine/DataTable.h"
+#include "Engine/DataAsset.h"
+#include "Indexers/DialogueWaveIndexer.h"
+#include "Sound/DialogueWave.h"
 
 #define LOCTEXT_NAMESPACE "FAssetSearchManager"
 
@@ -78,10 +84,11 @@ FAssetSearchManager::~FAssetSearchManager()
 
 void FAssetSearchManager::Start()
 {
-	RegisterIndexer(TEXT("DataAsset"), new FDataAssetIndexer());
-	RegisterIndexer(TEXT("DataTable"), new FDataTableIndexer());
-	RegisterIndexer(TEXT("Blueprint"), new FBlueprintIndexer());
-	RegisterIndexer(TEXT("WidgetBlueprint"), new FWidgetBlueprintIndexer());
+	RegisterAssetIndexer(UDataAsset::StaticClass(), MakeUnique<FDataAssetIndexer>());
+	RegisterAssetIndexer(UDataTable::StaticClass(), MakeUnique<FDataTableIndexer>());
+	RegisterAssetIndexer(UBlueprint::StaticClass(), MakeUnique<FBlueprintIndexer>());
+	RegisterAssetIndexer(UWidgetBlueprint::StaticClass(), MakeUnique<FWidgetBlueprintIndexer>());
+	RegisterAssetIndexer(UDialogueWave::StaticClass(), MakeUnique<FDialogueWaveIndexer>());
 
 	const FString SessionPath = FPaths::ConvertRelativePathToFull(FPaths::Combine(FPaths::ProjectSavedDir(), TEXT("Search")));
 	SearchDatabase.Open(SessionPath);
@@ -120,10 +127,11 @@ FSearchStats FAssetSearchManager::GetStats() const
 	return Stats;
 }
 
-void FAssetSearchManager::RegisterIndexer(FName AssetClassName, IAssetIndexer* Indexer)
+void FAssetSearchManager::RegisterAssetIndexer(const UClass* AssetClass, TUniquePtr<IAssetIndexer>&& Indexer)
 {
 	check(IsInGameThread());
-	Indexers.Add(AssetClassName, Indexer);
+
+	Indexers.Add(AssetClass->GetFName(), MoveTemp(Indexer));
 }
 
 void FAssetSearchManager::OnAssetAdded(const FAssetData& InAssetData)
@@ -325,12 +333,12 @@ FString FAssetSearchManager::GetDerivedDataKey(const FAssetData& UnindexedAsset)
 	return DDCKey;
 }
 
-bool FAssetSearchManager::HasIndexerForClass(UClass* AssetClass)
+bool FAssetSearchManager::HasIndexerForClass(const UClass* AssetClass)
 {
-	UClass* IndexableClass = AssetClass;
+	const UClass* IndexableClass = AssetClass;
 	while (IndexableClass)
 	{
-		if (IAssetIndexer* Indexer = Indexers.FindRef(IndexableClass->GetFName()))
+		if (Indexers.Contains(IndexableClass->GetFName()))
 		{
 			return true;
 		}
@@ -357,8 +365,10 @@ void FAssetSearchManager::StoreIndexForAsset(UObject* InAsset, bool bLegacyIndex
 			UClass* IndexableClass = InAsset->GetClass();
 			while (IndexableClass)
 			{
-				if (IAssetIndexer* Indexer = Indexers.FindRef(IndexableClass->GetFName()))
+				if (TUniquePtr<IAssetIndexer>* IndexerPtr = Indexers.Find(IndexableClass->GetFName()))
 				{
+					IAssetIndexer* Indexer = IndexerPtr->Get();
+
 					bWasIndexed = true;
 					Serializer.BeginIndexer(Indexer);
 					Indexer->IndexAsset(InAsset, Serializer);
diff --git a/Engine/Plugins/Editor/AssetSearch/Source/Private/AssetSearchManager.h b/Engine/Plugins/Editor/AssetSearch/Source/Private/AssetSearchManager.h
--- a/Engine/Plugins/Editor/AssetSearch/Source/Private/AssetSearchManager.h
+++ b/Engine/Plugins/Editor/AssetSearch/Source/Private/AssetSearchManager.h
@@ -9,6 +9,7 @@
 #include "HAL/Runnable.h"
 
 class FRunnableThread;
+class UClass;
 
 class FAssetSearchManager : public FRunnable
 {
@@ -17,7 +18,7 @@ public:
 	~FAssetSearchManager();
 	
 	void Start();
-	void RegisterIndexer(FName AssetClassName, IAssetIndexer* Indexer);
+	void RegisterAssetIndexer(const UClass* AssetClass, TUniquePtr<IAssetIndexer>&& Indexer);
 
 	FSearchStats GetStats() const;
 
@@ -49,7 +50,7 @@ private:
 	FString TryGetDDCKeyForAsset(const FAssetData& InAsset);
 	FString GetDerivedDataKey(const FSHAHash& IndexedContentHash);
 	FString GetDerivedDataKey(const FAssetData& UnindexedAsset);
-	bool HasIndexerForClass(UClass* AssetClass);
+	bool HasIndexerForClass(const UClass* AssetClass);
 	void StoreIndexForAsset(UObject* InAsset, bool Unindexed);
 	void LoadDDCContentIntoDatabase(const FAssetData& InAsset, const TArray<uint8>& Content, const FString& DerivedDataKey);
 
@@ -71,7 +72,7 @@ private:
 	double LastRecordCountUpdateSeconds;
 
 private:
-	TMap<FName, IAssetIndexer*> Indexers;
+	TMap<FName, TUniquePtr<IAssetIndexer>> Indexers;
 
 	TArray<TWeakObjectPtr<UObject>> RequestIndexQueue;
 
diff --git a/Engine/Plugins/Editor/AssetSearch/Source/Private/AssetSearchModule.cpp b/Engine/Plugins/Editor/AssetSearch/Source/Private/AssetSearchModule.cpp
--- a/Engine/Plugins/Editor/AssetSearch/Source/Private/AssetSearchModule.cpp
+++ b/Engine/Plugins/Editor/AssetSearch/Source/Private/AssetSearchModule.cpp
@@ -76,14 +76,9 @@ public:
 		SearchManager->ForceIndexOnAssetsMissingIndex();
 	}
 
-	virtual void RegisterIndexer(FName AssetClassName, IAssetIndexer* Indexer) override
+	virtual void RegisterAssetIndexer(const UClass* InAssetClass, TUniquePtr<IAssetIndexer>&& Indexer) override
 	{
-		SearchManager->RegisterIndexer(AssetClassName, Indexer);
-	}
-
-	virtual void UnregisterIndexer(IAssetIndexer* Indexer) override
-	{
-		
+		SearchManager->RegisterAssetIndexer(InAssetClass, MoveTemp(Indexer));
 	}
 
 private:
diff --git a/Engine/Plugins/Editor/AssetSearch/Source/Private/Indexers/DialogueWaveIndexer.cpp b/Engine/Plugins/Editor/AssetSearch/Source/Private/Indexers/DialogueWaveIndexer.cpp
new file mode 100644
--- /dev/null
+++ b/Engine/Plugins/Editor/AssetSearch/Source/Private/Indexers/DialogueWaveIndexer.cpp
@@ -0,0 +1,31 @@
+// Copyright Epic Games, Inc. All Rights Reserved.
+
+#include "DialogueWaveIndexer.h"
+#include "Utility/IndexerUtilities.h"
+#include "Engine/DataAsset.h"
+#include "Sound/DialogueWave.h"
+
+enum class EDialogueWaveIndexerVersion
+{
+	Empty = 0,
+	Initial = 1,
+
+	Current = Initial,
+};
+
+int32 FDialogueWaveIndexer::GetVersion() const
+{
+	return (int32)EDialogueWaveIndexerVersion::Current;
+}
+
+void FDialogueWaveIndexer::IndexAsset(const UObject* InAssetObject, FSearchSerializer& Serializer)
+{
+	const UDialogueWave* DialogueWave = Cast<UDialogueWave>(InAssetObject);
+	check(DialogueWave);
+
+	Serializer.BeginIndexingObject(DialogueWave, TEXT("$self"));
+
+	Serializer.IndexProperty(TEXT("SpokenText"), DialogueWave->SpokenText);
+
+	Serializer.EndIndexingObject();
+}
\ No newline at end of file
diff --git a/Engine/Plugins/Editor/AssetSearch/Source/Private/Indexers/DialogueWaveIndexer.h b/Engine/Plugins/Editor/AssetSearch/Source/Private/Indexers/DialogueWaveIndexer.h
new file mode 100644
--- /dev/null
+++ b/Engine/Plugins/Editor/AssetSearch/Source/Private/Indexers/DialogueWaveIndexer.h
@@ -0,0 +1,13 @@
+// Copyright Epic Games, Inc. All Rights Reserved.
+
+#pragma once
+
+#include "CoreMinimal.h"
+#include "IAssetSearchModule.h"
+
+class FDialogueWaveIndexer : public IAssetIndexer
+{
+	virtual FString GetName() const override { return TEXT("DialogueWave"); }
+	virtual int32 GetVersion() const override;
+	virtual void IndexAsset(const UObject* InAssetObject, FSearchSerializer& Serializer) override;
+};
\ No newline at end of file
diff --git a/Engine/Plugins/Editor/AssetSearch/Source/Public/IAssetSearchModule.h b/Engine/Plugins/Editor/AssetSearch/Source/Public/IAssetSearchModule.h
--- a/Engine/Plugins/Editor/AssetSearch/Source/Public/IAssetSearchModule.h
+++ b/Engine/Plugins/Editor/AssetSearch/Source/Public/IAssetSearchModule.h
@@ -53,6 +53,8 @@ public:
 	virtual void IndexAsset(const UObject* InAssetObject, FSearchSerializer& Serializer) = 0;
 };
 
+class UClass;
+
 /**
  *
  */
@@ -77,8 +79,7 @@ public:
 
 	virtual void ForceIndexOnAssetsMissingIndex() = 0;
 
-	virtual void RegisterIndexer(FName AssetClassName, IAssetIndexer* Indexer) = 0;
-	virtual void UnregisterIndexer(IAssetIndexer* Indexer) = 0;
+	virtual void RegisterAssetIndexer(const UClass* AssetClass, TUniquePtr<IAssetIndexer>&& Indexer) = 0;
 
 public:
 
