Build fix for some clang based editor compile errors

diff --git a/Engine/Plugins/Editor/AssetSearch/Source/Public/IAssetSearchModule.h b/Engine/Plugins/Editor/AssetSearch/Source/Public/IAssetSearchModule.h
--- a/Engine/Plugins/Editor/AssetSearch/Source/Public/IAssetSearchModule.h
+++ b/Engine/Plugins/Editor/AssetSearch/Source/Public/IAssetSearchModule.h
@@ -48,11 +48,15 @@ struct FSearchStats
 class IAssetIndexer
 {
 public:
+
+	virtual ~IAssetIndexer() = 0;
 	virtual FString GetName() const = 0;
 	virtual int32 GetVersion() const = 0;
 	virtual void IndexAsset(const UObject* InAssetObject, FSearchSerializer& Serializer) = 0;
 };
 
+inline IAssetIndexer::~IAssetIndexer() = default;
+
 class UClass;
 
 /**
