Universal Search now covers, Sound Cues, Dialogue Waves, Levels (including the Level Blueprint).  Comments on blueprint nodes are now indexed.

Now hashing all assets that are indexed.  We now only generate a single DDC key.  AssetSearch_V{SerializerVersion}_{IndexersNamesAndVersions}_{ObjectPathHash}_{FileOnDiskHash}

The file hashes along with last modified are stored in the asset database to make quick lookup easy.

A project now has control over which directories to not index.

diff --git a/Engine/Plugins/Editor/AssetSearch/Source/Private/AssetSearchDatabase.cpp b/Engine/Plugins/Editor/AssetSearch/Source/Private/AssetSearchDatabase.cpp
--- a/Engine/Plugins/Editor/AssetSearch/Source/Private/AssetSearchDatabase.cpp
+++ b/Engine/Plugins/Editor/AssetSearch/Source/Private/AssetSearchDatabase.cpp
@@ -10,6 +10,9 @@
 
 #include "AssetData.h"
 #include "Misc/TextFilterExpressionEvaluator.h"
+#include "Misc/PackageName.h"
+#include "Misc/Paths.h"
+#include "Misc/SecureHash.h"
 
 PRAGMA_DISABLE_OPTIMIZATION
 
@@ -20,6 +23,7 @@ enum class EAssetSearchDatabaseVersion
 	Empty = 0,
 	Initial = 1,
 	IndexingAssetIdsAssetPathsUnique = 2,
+	IntroducingFileHashing = 2,
 
 	// -----<new versions can be added above this line>-------------------------------------------------
 	VersionPlusOne,
@@ -60,6 +64,11 @@ public:
 
 		PREPARE_STATEMENT(Statement_SearchAssetsFTS);
 
+		PREPARE_STATEMENT(Statement_AddFileInfo);
+		PREPARE_STATEMENT(Statement_UpdateFileInfo);
+		PREPARE_STATEMENT(Statement_GetFileInfo);
+		PREPARE_STATEMENT(Statement_GetAllFileInfos);
+
 #undef PREPARE_STATEMENT
 
 		return true;
@@ -148,6 +157,109 @@ public:
 		return INDEX_NONE;
 	}
 
+	struct FCachedFileInfo
+	{
+		int64 FileId;
+		FString FilePath;
+		FDateTime LastModifed;
+		FString Hash;
+
+		FAssetFileInfo ToAssetFileInfo() const
+		{
+			FAssetFileInfo AssetFileInfo;
+			AssetFileInfo.LastModified = LastModifed;
+			AssetFileInfo.PackageName = *FilePath;
+			LexFromString(AssetFileInfo.Hash, *Hash);
+
+			return AssetFileInfo;
+		}
+	};
+
+	SQLITE_PREPARED_STATEMENT(FGetFileInfo,
+		"SELECT fileid, file_last_modified, file_hash FROM table_files WHERE file_path = ?1;",
+		SQLITE_PREPARED_STATEMENT_COLUMNS(int64, FDateTime, FString),
+		SQLITE_PREPARED_STATEMENT_BINDINGS(FString)
+	);
+	private: FGetFileInfo Statement_GetFileInfo;
+	public: bool GetFileInfo(const FString& InFullFilePath, FCachedFileInfo& OutFileInfo)
+	{
+		OutFileInfo.FilePath = InFullFilePath.ToLower();
+		if (Statement_GetFileInfo.BindAndExecuteSingle(OutFileInfo.FilePath, OutFileInfo.FileId, OutFileInfo.LastModifed, OutFileInfo.Hash))
+		{
+			return true;
+		}
+
+		return false;
+	}
+
+	SQLITE_PREPARED_STATEMENT(FGetAllFileInfos,
+		"SELECT file_path, file_last_modified, file_hash FROM table_files;",
+		SQLITE_PREPARED_STATEMENT_COLUMNS(FString, FDateTime, FString),
+		SQLITE_PREPARED_STATEMENT_BINDINGS()
+	);
+	FGetAllFileInfos Statement_GetAllFileInfos;
+	bool GetAllFileInfos(TFunctionRef<ESQLitePreparedStatementExecuteRowResult(FAssetFileInfo&&)> InCallback)
+	{
+		return Statement_GetAllFileInfos.BindAndExecute([&InCallback](const FGetAllFileInfos& InStatement)
+		{
+			FCachedFileInfo FileInfo;
+			if (InStatement.GetColumnValues(FileInfo.FilePath, FileInfo.LastModifed, FileInfo.Hash))
+			{
+				return InCallback(FileInfo.ToAssetFileInfo());
+			}
+			return ESQLitePreparedStatementExecuteRowResult::Error;
+		}) != INDEX_NONE;
+	}
+
+	SQLITE_PREPARED_STATEMENT_BINDINGS_ONLY(
+		FUpdateFileInfo,
+		" UPDATE table_files SET file_last_modified = ?2, file_hash = ?3 WHERE file_path = ?1;",
+		SQLITE_PREPARED_STATEMENT_BINDINGS(FString, FDateTime, FString)
+	);
+	private: FUpdateFileInfo Statement_UpdateFileInfo;
+	SQLITE_PREPARED_STATEMENT_BINDINGS_ONLY(
+		FAddFileInfo,
+		" INSERT INTO table_files(file_path, file_last_modified, file_hash)"
+		" VALUES(?1, ?2, ?3);",
+		SQLITE_PREPARED_STATEMENT_BINDINGS(FString, FDateTime, FString)
+	);
+	private: FAddFileInfo Statement_AddFileInfo;
+	public: bool AddOrUpdateFileInfo(const FAssetData& InAssetData, FAssetFileInfo& OutFileInfo)
+	{
+		const FString PackageName = InAssetData.PackageName.ToString();
+		const bool bIsWorldAsset = (InAssetData.AssetClass == UWorld::StaticClass()->GetFName());
+		const FString Extension = bIsWorldAsset ? FPackageName::GetMapPackageExtension() : FPackageName::GetAssetPackageExtension();
+		const FString FilePath = FPackageName::LongPackageNameToFilename(PackageName, Extension);
+		const FString FullFilePath = FPaths::ConvertRelativePathToFull(FilePath);
+
+		IPlatformFile& PlatformFile = FPlatformFileManager::Get().GetPlatformFile();
+		const FDateTime CurrentLastModified = PlatformFile.GetTimeStamp(*FullFilePath);
+
+		FCachedFileInfo FileInfo;
+		if (GetFileInfo(PackageName.ToLower(), FileInfo))
+		{
+			if (CurrentLastModified == FileInfo.LastModifed)
+			{
+				OutFileInfo = FileInfo.ToAssetFileInfo();
+				return false;
+			}
+
+			OutFileInfo = FileInfo.ToAssetFileInfo();
+			OutFileInfo.LastModified = CurrentLastModified;
+			OutFileInfo.Hash = FMD5Hash::HashFile(*FullFilePath);
+			Statement_UpdateFileInfo.BindAndExecuteSingle(PackageName.ToLower(), OutFileInfo.LastModified, LexToString(OutFileInfo.Hash));
+			return true;
+		}
+		else
+		{
+			OutFileInfo = FileInfo.ToAssetFileInfo();
+			OutFileInfo.LastModified = CurrentLastModified;
+			OutFileInfo.Hash = FMD5Hash::HashFile(*FullFilePath);
+			Statement_AddFileInfo.BindAndExecuteSingle(PackageName.ToLower(), OutFileInfo.LastModified, LexToString(OutFileInfo.Hash));
+			return true;
+		}
+	}
+
 	SQLITE_PREPARED_STATEMENT_BINDINGS_ONLY(
 		FAddAssetToAssetTable,
 		" INSERT INTO table_assets(asset_name, asset_class, asset_path, index_hash)"
@@ -534,6 +646,13 @@ bool FAssetSearchDatabase::Open(const FString& InSessionPath, const ESQLiteDatab
 
 	// Create our required tables
 	//========================================================================
+	if (!ensure(Database->Execute(TEXT("CREATE TABLE IF NOT EXISTS table_files(fileid INTEGER PRIMARY KEY, file_path TEXT UNIQUE, file_last_modified INTEGER NOT NULL, file_hash);"))))
+	{
+		LogLastError();
+		Close();
+		return false;
+	}
+
 	if (!ensure(Database->Execute(TEXT("CREATE TABLE IF NOT EXISTS table_assets(assetid INTEGER PRIMARY KEY, asset_name, asset_class, asset_path TEXT UNIQUE, index_hash);"))))
 	{
 		LogLastError();
@@ -638,6 +757,15 @@ bool FAssetSearchDatabase::Open(const FString& InSessionPath, const ESQLiteDatab
 		return false;
 	}
 
+	if (!ensure(Database->Execute(
+		TEXT("CREATE UNIQUE INDEX IF NOT EXISTS file_path_index ON table_files(file_path);")
+	)))
+	{
+		LogLastError();
+		Close();
+		return false;
+	}
+
 	if (!ensure(Database->Execute(
 		TEXT("CREATE UNIQUE INDEX IF NOT EXISTS asset_path_index ON table_assets(asset_path);")
 	)))
@@ -717,6 +845,16 @@ void FAssetSearchDatabase::LogLastError() const
 	UE_LOG(LogAssetSearch, Error, TEXT("Database Error: %s"), *SessionPath, *GetLastError());
 }
 
+bool FAssetSearchDatabase::AddOrUpdateFileInfo(const FAssetData& InAssetData, FAssetFileInfo& OutFileInfo)
+{
+	if (ensure(Statements))
+	{
+		return Statements->AddOrUpdateFileInfo(InAssetData, OutFileInfo);
+	}
+
+	return false;
+}
+
 bool FAssetSearchDatabase::IsAssetUpToDate(const FAssetData& InAssetData, const FString& IndexedJsonHash)
 {
 	if (ensure(Statements))
@@ -773,11 +911,48 @@ void FAssetSearchDatabase::RemoveAsset(const FAssetData& InAssetData)
 	}
 }
 
+void FAssetSearchDatabase::AddOrUpdateFileInfos(const TArray<FAssetData>& InAssets)
+{
+	for (const FAssetData& InAsset : InAssets)
+	{
+		// If it's a redirector act like it has been removed from the system,
+		// we don't want old duplicate entries for it.
+		if (InAsset.IsRedirector())
+		{
+			continue;
+		}
+
+		// Freshen hash cache
+		FAssetFileInfo FileInfo;
+		AddOrUpdateFileInfo(InAsset, FileInfo);
+	}
+}
+
+TMap<FName, FAssetFileInfo> FAssetSearchDatabase::GetAllFileInfos()
+{
+	TMap<FName, FAssetFileInfo> FileInfos;
+	Statements->GetAllFileInfos([&FileInfos](FAssetFileInfo&& InResult)
+	{
+		FileInfos.Add(InResult.PackageName, InResult);
+
+		return ESQLitePreparedStatementExecuteRowResult::Continue;
+	});
+
+	return FileInfos;
+}
+
 void FAssetSearchDatabase::RemoveAssetsNotInThisSet(const TArray<FAssetData>& InAssets)
 {
 	TSet<FString> InAssetPaths;
 	for (const FAssetData& InAsset : InAssets)
 	{
+		// If it's a redirector act like it has been removed from the system,
+		// we don't want old duplicate entries for it.
+		if (InAsset.IsRedirector())
+		{
+			continue;
+		}
+
 		InAssetPaths.Add(InAsset.ObjectPath.ToString());
 	}
 
@@ -787,13 +962,12 @@ void FAssetSearchDatabase::RemoveAssetsNotInThisSet(const TArray<FAssetData>& In
 	{
 		if (!InAssetPaths.Contains(InResult))
 		{
-			MissingAssets.Add(InResult);
+			MissingAssets.Emplace(InResult);
 		}
 
 		return ESQLitePreparedStatementExecuteRowResult::Continue;
 	});
 
-
 	for (const FString& MissingAsset : MissingAssets)
 	{
 		Statements->DeleteEntriesForAsset(MissingAsset);
diff --git a/Engine/Plugins/Editor/AssetSearch/Source/Private/AssetSearchDatabase.h b/Engine/Plugins/Editor/AssetSearch/Source/Private/AssetSearchDatabase.h
--- a/Engine/Plugins/Editor/AssetSearch/Source/Private/AssetSearchDatabase.h
+++ b/Engine/Plugins/Editor/AssetSearch/Source/Private/AssetSearchDatabase.h
@@ -10,6 +10,13 @@ enum class ESQLiteDatabaseOpenMode : uint8;
 class FAssetSearchDatabaseStatements;
 struct FAssetData;
 
+struct FAssetFileInfo
+{
+	FName PackageName;
+	FDateTime LastModified;
+	FMD5Hash Hash;
+};
+
 class FAssetSearchDatabase
 {
 public:
@@ -64,6 +71,16 @@ public:
 	 */
 	FString GetLastError() const;
 
+	/**
+	 * Update the file hash table.
+	 */
+	bool AddOrUpdateFileInfo(const FAssetData& InAssetData, FAssetFileInfo& OutFileInfo);
+	void AddOrUpdateFileInfos(const TArray<FAssetData>& InAssets);
+
+	void UpdateFileHash(const FAssetData& InAssetData, FAssetFileInfo& OutFileInfo);
+
+	TMap<FName, FAssetFileInfo> GetAllFileInfos();
+
 	/**
 	 * Is the asset cache up to date for this asset?
 	 */
diff --git a/Engine/Plugins/Editor/AssetSearch/Source/Private/AssetSearchManager.cpp b/Engine/Plugins/Editor/AssetSearch/Source/Private/AssetSearchManager.cpp
--- a/Engine/Plugins/Editor/AssetSearch/Source/Private/AssetSearchManager.cpp
+++ b/Engine/Plugins/Editor/AssetSearch/Source/Private/AssetSearchManager.cpp
@@ -24,6 +24,13 @@
 #include "Engine/DataAsset.h"
 #include "Indexers/DialogueWaveIndexer.h"
 #include "Sound/DialogueWave.h"
+#include "Indexers/LevelIndexer.h"
+#include "Settings/AssetSearchDeveloperSettings.h"
+#include "Indexers/SoundCueIndexer.h"
+#include "Sound/SoundCue.h"
+#include "Misc/ScopedSlowTask.h"
+
+PRAGMA_DISABLE_OPTIMIZATION
 
 #define LOCTEXT_NAMESPACE "FAssetSearchManager"
 
@@ -89,13 +96,15 @@ void FAssetSearchManager::Start()
 	RegisterAssetIndexer(UBlueprint::StaticClass(), MakeUnique<FBlueprintIndexer>());
 	RegisterAssetIndexer(UWidgetBlueprint::StaticClass(), MakeUnique<FWidgetBlueprintIndexer>());
 	RegisterAssetIndexer(UDialogueWave::StaticClass(), MakeUnique<FDialogueWaveIndexer>());
+	RegisterAssetIndexer(UWorld::StaticClass(), MakeUnique<FLevelIndexer>());
+	RegisterAssetIndexer(USoundCue::StaticClass(), MakeUnique<FSoundCueIndexer>());
+
 
 	const FString SessionPath = FPaths::ConvertRelativePathToFull(FPaths::Combine(FPaths::ProjectSavedDir(), TEXT("Search")));
 	SearchDatabase.Open(SessionPath);
 
 	FCoreUObjectDelegates::OnObjectSaved.AddRaw(this, &FAssetSearchManager::OnObjectSaved);
 	FCoreUObjectDelegates::OnAssetLoaded.AddRaw(this, &FAssetSearchManager::OnAssetLoaded);
-	UObject::FAssetRegistryTag::OnGetExtraObjectTags.AddRaw(this, &FAssetSearchManager::OnGetAssetTags);
 
 	IAssetRegistry& AssetRegistry = FModuleManager::LoadModuleChecked<FAssetRegistryModule>("AssetRegistry").Get();
 	AssetRegistry.OnAssetAdded().AddRaw(this, &FAssetSearchManager::OnAssetAdded);
@@ -151,6 +160,22 @@ void FAssetSearchManager::OnAssetAdded(const FAssetData& InAssetData)
 		}
 	}
 
+	// 
+	const UAssetSearchDeveloperSettings* Settings = GetDefault<UAssetSearchDeveloperSettings>();
+	for (const FDirectoryPath& IgnoredPath : Settings->IgnoredPaths)
+	{
+		if (PackageName.StartsWith(IgnoredPath.Path))
+		{
+			return;
+		}
+	}
+
+	// Don't index redirectors, just act like they don't exist.
+	if (InAssetData.IsRedirector())
+	{
+		return;
+	}
+
 	FAssetOperation Operation;
 	Operation.Asset = InAssetData;
 	ProcessAssetQueue.Add(Operation);
@@ -173,10 +198,12 @@ void FAssetSearchManager::OnAssetScanFinished()
 	TArray<FAssetData> AllAssets;
 	IAssetRegistry& AssetRegistry = FModuleManager::LoadModuleChecked<FAssetRegistryModule>("AssetRegistry").Get();
 	AssetRegistry.GetAllAssets(AllAssets, false);
-
+	
+	PendingDatabaseUpdates++;
 	UpdateOperations.Enqueue([this, AssetsAvailable = MoveTemp(AllAssets)]() mutable {
 		FScopeLock ScopedLock(&SearchDatabaseCS);
 		SearchDatabase.RemoveAssetsNotInThisSet(AssetsAvailable);
+		PendingDatabaseUpdates--;
 	});
 }
 
@@ -186,7 +213,7 @@ void FAssetSearchManager::OnObjectSaved(UObject* InObject)
 
 	if (!GIsCookerLoadingPackage)
 	{
-		StoreIndexForAsset(InObject, false);
+		RequestIndexAsset(InObject);
 	}
 }
 
@@ -200,46 +227,34 @@ void FAssetSearchManager::OnAssetLoaded(UObject* InObject)
 	}
 }
 
-void FAssetSearchManager::OnGetAssetTags(const UObject* Object, TArray<UObject::FAssetRegistryTag>& OutTags)
+bool FAssetSearchManager::RequestIndexAsset(UObject* InAsset)
 {
 	check(IsInGameThread());
 
-	const FString ObjectPath = Object->GetPathName();
-	const FContentHashEntry* ContentEntry = ContentHashCache.Find(ObjectPath);
-
-	if (ContentEntry)
+	if (GEditor->IsAutosaving())
 	{
-		static FName SearchIndexContentTag(TEXT("SearchIndexContent"));
-		static FName SearchIndexContentKeyTag(TEXT("SearchIndexContentKey"));
-		OutTags.Add(UObject::FAssetRegistryTag(SearchIndexContentKeyTag, *ContentEntry->ContentHash, UObject::FAssetRegistryTag::TT_Hidden));
-		OutTags.Add(UObject::FAssetRegistryTag(SearchIndexContentTag, *ContentEntry->Content, UObject::FAssetRegistryTag::TT_Hidden));
+		return false;
 	}
-}
-
-void FAssetSearchManager::AddToContentTagCache(const FAssetData& InAsset, const FString& InContent, const FString& InContentHash)
-{
-	FContentHashEntry Entry;
-	Entry.Content = InContent;
-	Entry.ContentHash = InContentHash;
-	ContentHashCache.Add(InAsset.ObjectPath.ToString(), Entry);
-}
 
-void FAssetSearchManager::RequestIndexAsset(UObject* InAsset)
-{
-	TWeakObjectPtr<UObject> AssetWeakPtr = InAsset;
-
-	FAssetData AssetData(InAsset);
-	FString AssetJsonDDCKey = TryGetDDCKeyForAsset(InAsset);
+	if (IsAssetIndexable(InAsset))
+	{
+		TWeakObjectPtr<UObject> AssetWeakPtr = InAsset;
+		FAssetData AssetData(InAsset);
 
-	UpdateOperations.Enqueue([this, AssetData, AssetWeakPtr, AssetJsonDDCKey]() {
-		FScopeLock ScopedLock(&SearchDatabaseCS);
-		if (!SearchDatabase.IsAssetUpToDate(AssetData, AssetJsonDDCKey))
-		{
-			AsyncTask(ENamedThreads::GameThread, [this, AssetWeakPtr]() {
-				StoreIndexForAsset(AssetWeakPtr.Get(), true);
+		return AsyncGetDerivedDataKey(AssetData, [this, AssetData, AssetWeakPtr](FString InDDCKey) {
+			UpdateOperations.Enqueue([this, AssetData, AssetWeakPtr, InDDCKey]() {
+				FScopeLock ScopedLock(&SearchDatabaseCS);
+				if (!SearchDatabase.IsAssetUpToDate(AssetData, InDDCKey))
+				{
+					AsyncTask(ENamedThreads::GameThread, [this, AssetWeakPtr]() {
+						StoreIndexForAsset(AssetWeakPtr.Get());
+					});
+				}
 			});
-		}
-	});
+		});
+	}
+
+	return false;
 }
 
 bool FAssetSearchManager::IsAssetIndexable(UObject* InAsset)
@@ -248,7 +263,7 @@ bool FAssetSearchManager::IsAssetIndexable(UObject* InAsset)
 	{
 		// If it's not a permanent package, and one we just loaded for diffing, don't index it.
 		UPackage* Package = InAsset->GetOutermost();
-		if (Package->HasAnyPackageFlags(LOAD_ForDiff | LOAD_PackageForPIE | LOAD_ForFileDiff))
+		if (Package->HasAnyPackageFlags(/*LOAD_ForDiff | */LOAD_PackageForPIE | LOAD_ForFileDiff))
 		{
 			return false;
 		}
@@ -264,78 +279,76 @@ bool FAssetSearchManager::IsAssetIndexable(UObject* InAsset)
 	return false;
 }
 
-FString FAssetSearchManager::TryGetDDCKeyForAsset(const FAssetData& InAsset)
+bool FAssetSearchManager::TryLoadIndexForAsset(const FAssetData& InAssetData)
 {
-	FString AssetJsonDDCKey;
-
-	{
-		FString ContentDDCKey;
-		static FName SearchIndexContentKeyTag(TEXT("SearchIndexContentKey"));
-		if (InAsset.GetTagValue(SearchIndexContentKeyTag, ContentDDCKey))
-		{
-			AssetJsonDDCKey = ContentDDCKey;
-		}
-		else
-		{
-			UClass* AssetClass = InAsset.GetClass();
-			if (HasIndexerForClass(AssetClass))
-			{
-				const FString UnindexedAssetKey = GetDerivedDataKey(InAsset);
-				AssetJsonDDCKey = UnindexedAssetKey;
-			}
-		}
-	}
-
-	return AssetJsonDDCKey;
-}
-
-bool FAssetSearchManager::TryLoadIndexForAsset(const FAssetData& InAsset)
-{
-	check(IsInGameThread());
-	FString AssetJsonDDCKey = TryGetDDCKeyForAsset(InAsset);
-
-	if (!AssetJsonDDCKey.IsEmpty())
-	{
-		FeedOperations.Enqueue([this, InAsset, AssetJsonDDCKey]() {
+	return AsyncGetDerivedDataKey(InAssetData, [this, InAssetData](FString InDDCKey) {
+		FeedOperations.Enqueue([this, InAssetData, InDDCKey]() {
 			FScopeLock ScopedLock(&SearchDatabaseCS);
-			if (!SearchDatabase.IsAssetUpToDate(InAsset, AssetJsonDDCKey))
+			if (!SearchDatabase.IsAssetUpToDate(InAssetData, InDDCKey))
 			{
 				PendingDownloads++;
 
 				FAssetDDCRequest DDCRequest;
-				DDCRequest.AssetData = InAsset;
-				DDCRequest.DDCKey_IndexDataHash = AssetJsonDDCKey;
-				DDCRequest.DDCHandle = GetDerivedDataCacheRef().GetAsynchronous(*AssetJsonDDCKey, InAsset.ObjectPath.ToString());
+				DDCRequest.AssetData = InAssetData;
+				DDCRequest.DDCKey_IndexDataHash = InDDCKey;
+				DDCRequest.DDCHandle = GetDerivedDataCacheRef().GetAsynchronous(*InDDCKey, InAssetData.ObjectPath.ToString());
 				ProcessDDCQueue.Enqueue(DDCRequest);
 			}
 		});
-
-		return true;
-	}
-
-	return false;
+	});
 }
 
-FString FAssetSearchManager::GetDerivedDataKey(const FSHAHash& IndexedContentHash)
+bool FAssetSearchManager::AsyncGetDerivedDataKey(const FAssetData& InAssetData, TFunction<void(FString)> DDCKeyCallback)
 {
-	const FString DDCKey = TEXT("AssetSearch_B") + LexToString(FSearchSerializer::GetVersion()) + TEXT("_") + IndexedContentHash.ToString();
-	return DDCKey;
-}
+	check(IsInGameThread());
 
-FString FAssetSearchManager::GetDerivedDataKey(const FAssetData& UnindexedAsset)
-{
-	FString ContentPath = UnindexedAsset.ObjectPath.ToString();
+	FString IndexersNamesAndVersions = GetIndexerVersion(InAssetData.GetClass());
 
-	FSHAHash UnindexedAssetHash;
-	FSHA1::HashBuffer(*ContentPath, ContentPath.Len() * sizeof(FString::ElementType), UnindexedAssetHash.Hash);
+	// If the indexer names and versions is empty, then we know it's not possible to index this type of thing.
+	if (IndexersNamesAndVersions.IsEmpty())
+	{
+		return false;
+	}
 
-	const FString DDCKey = TEXT("AssetSearch_Legacy_B") + LexToString(FSearchSerializer::GetVersion()) + TEXT("_") + UnindexedAssetHash.ToString();
-	return DDCKey;
+	UpdateOperations.Enqueue([this, InAssetData, IndexersNamesAndVersions, DDCKeyCallback]() {
+		FScopeLock ScopedLock(&SearchDatabaseCS);
+
+		FAssetFileInfo FileInfo;
+		SearchDatabase.AddOrUpdateFileInfo(InAssetData, FileInfo);
+
+		if (FileInfo.Hash.IsValid())
+		{
+			// The universal key for content is:
+			// AssetSearch_V{SerializerVersion}_{IndexersNamesAndVersions}_{ObjectPathHash}_{FileOnDiskHash}
+
+			const FString ObjectPathString = InAssetData.ObjectPath.ToString();
+
+			FSHAHash ObjectPathHash;
+			FSHA1::HashBuffer(*ObjectPathString, ObjectPathString.Len() * sizeof(FString::ElementType), ObjectPathHash.Hash);
+
+			TStringBuilder<512> DDCKey;
+			DDCKey.Append(TEXT("AssetSearch_V"));
+			DDCKey.Append(LexToString(FSearchSerializer::GetVersion()));
+			DDCKey.Append(TEXT("_"));
+			DDCKey.Append(IndexersNamesAndVersions);
+			DDCKey.Append(TEXT("_"));
+			DDCKey.Append(ObjectPathHash.ToString());
+			DDCKey.Append(TEXT("_"));
+			DDCKey.Append(LexToString(FileInfo.Hash));
+
+			const FString DDCKeyString = DDCKey.ToString();
+			AsyncTask(ENamedThreads::GameThread, [this, DDCKeyString, DDCKeyCallback]() {
+				DDCKeyCallback(DDCKeyString);
+			});
+		}
+	});
+
+	return true;
 }
 
-bool FAssetSearchManager::HasIndexerForClass(const UClass* AssetClass)
+bool FAssetSearchManager::HasIndexerForClass(const UClass* InAssetClass) const
 {
-	const UClass* IndexableClass = AssetClass;
+	const UClass* IndexableClass = InAssetClass;
 	while (IndexableClass)
 	{
 		if (Indexers.Contains(IndexableClass->GetFName()))
@@ -349,51 +362,62 @@ bool FAssetSearchManager::HasIndexerForClass(const UClass* AssetClass)
 	return false;
 }
 
-void FAssetSearchManager::StoreIndexForAsset(UObject* InAsset, bool bLegacyIndexing)
+FString FAssetSearchManager::GetIndexerVersion(const UClass* InAssetClass) const
+{
+	TStringBuilder<256> VersionString;
+
+	TArray<UClass*> NestedIndexedTypes;
+
+	const UClass* IndexableClass = InAssetClass;
+	while (IndexableClass)
+	{
+		if (const TUniquePtr<IAssetIndexer>* IndexerPtr = Indexers.Find(IndexableClass->GetFName()))
+		{
+			IAssetIndexer* Indexer = IndexerPtr->Get();
+			VersionString.Append(Indexer->GetName());
+			VersionString.Append(TEXT("_"));
+			VersionString.Append(LexToString(Indexer->GetVersion()));
+
+			Indexer->GetNestedAssetTypes(NestedIndexedTypes);
+		}
+
+		IndexableClass = IndexableClass->GetSuperClass();
+	}
+
+	for (UClass* NestedIndexedType : NestedIndexedTypes)
+	{
+		VersionString.Append(GetIndexerVersion(NestedIndexedType));
+	}
+
+	return VersionString.ToString();
+}
+
+void FAssetSearchManager::StoreIndexForAsset(UObject* InAsset)
 {
 	check(IsInGameThread());
 
-	if (IsAssetIndexable(InAsset))
+	if (IsAssetIndexable(InAsset) && HasIndexerForClass(InAsset->GetClass()))
 	{
 		FAssetData InAssetData(InAsset);
 
-		bool bWasIndexed = false;
 		FString IndexedJson;
+		bool bWasIndexed = false;
 		{
 			FSearchSerializer Serializer(InAssetData, &IndexedJson);
-
-			UClass* IndexableClass = InAsset->GetClass();
-			while (IndexableClass)
-			{
-				if (TUniquePtr<IAssetIndexer>* IndexerPtr = Indexers.Find(IndexableClass->GetFName()))
-				{
-					IAssetIndexer* Indexer = IndexerPtr->Get();
-
-					bWasIndexed = true;
-					Serializer.BeginIndexer(Indexer);
-					Indexer->IndexAsset(InAsset, Serializer);
-					Serializer.EndIndexer();
-				}
-
-				IndexableClass = IndexableClass->GetSuperClass();
-			}
+			bWasIndexed = Serializer.IndexAsset(InAsset, Indexers);
 		}
 
-		if (bWasIndexed)
+		if (bWasIndexed && !IndexedJson.IsEmpty())
 		{
-			// Hash the content so that we can store it in the DDC.
-			FSHAHash IndexedJsonHash;
-			FSHA1::HashBuffer(*IndexedJson, IndexedJson.Len() * sizeof(FString::ElementType), IndexedJsonHash.Hash);
-
-			const FString DerivedDataKey = bLegacyIndexing ? GetDerivedDataKey(InAssetData) : GetDerivedDataKey(IndexedJsonHash);
+			AsyncGetDerivedDataKey(InAssetData, [this, InAssetData, IndexedJson](FString InDDCKey) {
+				check(IsInGameThread());
 
-			FTCHARToUTF8 IndexedJsonUTF8(*IndexedJson);
-			TArrayView<const uint8> IndexedJsonUTF8View((const uint8*)IndexedJsonUTF8.Get(), IndexedJsonUTF8.Length() * sizeof(UTF8CHAR));
-			GetDerivedDataCacheRef().Put(*DerivedDataKey, IndexedJsonUTF8View, InAssetData.ObjectPath.ToString(), bLegacyIndexing);
+				FTCHARToUTF8 IndexedJsonUTF8(*IndexedJson);
+				TArrayView<const uint8> IndexedJsonUTF8View((const uint8*)IndexedJsonUTF8.Get(), IndexedJsonUTF8.Length() * sizeof(UTF8CHAR));
+				GetDerivedDataCacheRef().Put(*InDDCKey, IndexedJsonUTF8View, InAssetData.ObjectPath.ToString(), false);
 
-			AddToContentTagCache(InAssetData, IndexedJson, DerivedDataKey);
-
-			AddOrUpdateAsset(InAssetData, IndexedJson, DerivedDataKey);
+				AddOrUpdateAsset(InAssetData, IndexedJson, InDDCKey);
+			});
 		}
 	}
 }
@@ -422,6 +446,18 @@ bool FAssetSearchManager::Tick_GameThread(float DeltaTime)
 {
 	check(IsInGameThread());
 
+	//if (0)
+	//{
+	//	IAssetRegistry& AssetRegistry = FModuleManager::LoadModuleChecked<FAssetRegistryModule>("AssetRegistry").Get();
+	//	TArray<FAssetData> TempAssetData;
+	//	AssetRegistry.GetAllAssets(TempAssetData, true);
+
+	//	for (const FAssetData& Data : TempAssetData)
+	//	{
+	//		OnAssetAdded(Data);
+	//	}
+	//}
+
 	int32 ScanLimit = GameThread_AssetScanLimit;
 	while (ProcessAssetQueue.Num() > 0 && ScanLimit > 0 && PendingDownloads < PendingDownloadsMax)
 	{
@@ -430,9 +466,11 @@ bool FAssetSearchManager::Tick_GameThread(float DeltaTime)
 
 		if (Operation.bRemoval)
 		{
+			PendingDatabaseUpdates++;
 			UpdateOperations.Enqueue([this, Asset]() {
 				FScopeLock ScopedLock(&SearchDatabaseCS);
 				SearchDatabase.RemoveAsset(Asset);
+				PendingDatabaseUpdates--;
 			});
 		}
 		else
@@ -512,28 +550,27 @@ void FAssetSearchManager::ForceIndexOnAssetsMissingIndex()
 {
 	check(IsInGameThread());
 
-	GWarn->BeginSlowTask(LOCTEXT("ForceIndexOnAssetsMissingIndex", "Indexing Assets"), true, true);
+	FScopedSlowTask IndexingTask(FailedDDCRequests.Num(), LOCTEXT("ForceIndexOnAssetsMissingIndex", "Indexing Assets"));
+	IndexingTask.MakeDialog(true);
 
-	int32 RequestCount = 0;
+	int32 RemovedCount = 0;
 	for (const FAssetDDCRequest& Request : FailedDDCRequests)
 	{
-		if (GWarn->ReceivedUserCancel())
+		if (IndexingTask.ShouldCancel())
 		{
 			break;
 		}
 
+		IndexingTask.EnterProgressFrame(1, FText::Format(LOCTEXT("ForceIndexOnAssetsMissingIndexFormat", "Indexing Asset ({0} of {1})"), RemovedCount + 1, FailedDDCRequests.Num()));
 		if (UObject* AssetToIndex = Request.AssetData.GetAsset())
 		{
-			StoreIndexForAsset(AssetToIndex, Request.DDCKey_IndexDataHash.StartsWith(TEXT("AssetSearch_Legacy")));
+			StoreIndexForAsset(AssetToIndex);
 		}
 
-		RequestCount++;
-		GWarn->StatusForceUpdate(RequestCount, FailedDDCRequests.Num(), FText::FromString(Request.AssetData.PackageName.ToString()));
+		RemovedCount++;
 	}
 
-	GWarn->EndSlowTask();
-
-	FailedDDCRequests.Reset();
+	FailedDDCRequests.RemoveAtSwap(0, RemovedCount);
 }
 
 void FAssetSearchManager::Search(const FSearchQuery& Query, TFunction<void(TArray<FSearchRecord>&&)> InCallback)
@@ -562,4 +599,6 @@ void FAssetSearchManager::Search(const FSearchQuery& Query, TFunction<void(TArra
 	});
 }
 
-#undef LOCTEXT_NAMESPACE
\ No newline at end of file
+#undef LOCTEXT_NAMESPACE
+
+PRAGMA_ENABLE_OPTIMIZATION
\ No newline at end of file
diff --git a/Engine/Plugins/Editor/AssetSearch/Source/Private/AssetSearchManager.h b/Engine/Plugins/Editor/AssetSearch/Source/Private/AssetSearchManager.h
--- a/Engine/Plugins/Editor/AssetSearch/Source/Private/AssetSearchManager.h
+++ b/Engine/Plugins/Editor/AssetSearch/Source/Private/AssetSearchManager.h
@@ -15,7 +15,7 @@ class FAssetSearchManager : public FRunnable
 {
 public:
 	FAssetSearchManager();
-	~FAssetSearchManager();
+	virtual ~FAssetSearchManager();
 	
 	void Start();
 	void RegisterAssetIndexer(const UClass* AssetClass, TUniquePtr<IAssetIndexer>&& Indexer);
@@ -39,29 +39,18 @@ private:
 
 	void OnObjectSaved(UObject* InObject);
 	void OnAssetLoaded(UObject* InObject);
-	void OnGetAssetTags(const UObject* Object, TArray<UObject::FAssetRegistryTag>& OutTags);
-	void AddToContentTagCache(const FAssetData& InAsset, const FString& InContent, const FString& InContentHash);
 
 	void AddOrUpdateAsset(const FAssetData& InAsset, const FString& IndexedJson, const FString& DerivedDataKey);
 
-	void RequestIndexAsset(UObject* InAsset);
+	bool RequestIndexAsset(UObject* InAsset);
 	bool IsAssetIndexable(UObject* InAsset);
 	bool TryLoadIndexForAsset(const FAssetData& InAsset);
-	FString TryGetDDCKeyForAsset(const FAssetData& InAsset);
-	FString GetDerivedDataKey(const FSHAHash& IndexedContentHash);
-	FString GetDerivedDataKey(const FAssetData& UnindexedAsset);
-	bool HasIndexerForClass(const UClass* AssetClass);
-	void StoreIndexForAsset(UObject* InAsset, bool Unindexed);
+	bool AsyncGetDerivedDataKey(const FAssetData& UnindexedAsset, TFunction<void(FString)> DDCKeyCallback);
+	bool HasIndexerForClass(const UClass* InAssetClass) const;
+	FString GetIndexerVersion(const UClass* InAssetClass) const;
+	void StoreIndexForAsset(UObject* InAsset);
 	void LoadDDCContentIntoDatabase(const FAssetData& InAsset, const TArray<uint8>& Content, const FString& DerivedDataKey);
 
-private:
-	struct FContentHashEntry
-	{
-		FString ContentHash;
-		FString Content;
-	};
-	TMap<FString /*AssetPath*/, FContentHashEntry> ContentHashCache;
-
 private:
 	FAssetSearchDatabase SearchDatabase;
 	FCriticalSection SearchDatabaseCS;
@@ -96,6 +85,8 @@ private:
 
 	FDelegateHandle TickerHandle;
 
+	TQueue<TFunction<void()>> GT_Requests;
+
 private:
 	TAtomic<bool> RunThread;
 	FRunnableThread* DatabaseThread = nullptr;
diff --git a/Engine/Plugins/Editor/AssetSearch/Source/Private/Indexers/BlueprintIndexer.cpp b/Engine/Plugins/Editor/AssetSearch/Source/Private/Indexers/BlueprintIndexer.cpp
--- a/Engine/Plugins/Editor/AssetSearch/Source/Private/Indexers/BlueprintIndexer.cpp
+++ b/Engine/Plugins/Editor/AssetSearch/Source/Private/Indexers/BlueprintIndexer.cpp
@@ -18,21 +18,21 @@ enum class EBlueprintIndexerVersion
 	Empty = 0,
 	Initial = 1,
 
-	Current = Initial,
+	// -----<new versions can be added above this line>-------------------------------------------------
+	VersionPlusOne,
+	LatestVersion = VersionPlusOne - 1
 };
 
 int32 FBlueprintIndexer::GetVersion() const
 {
-	return (int32)EBlueprintIndexerVersion::Current;
+	return (int32)EBlueprintIndexerVersion::LatestVersion;
 }
 
-void FBlueprintIndexer::IndexAsset(const UObject* InAssetObject, FSearchSerializer& Serializer)
+void FBlueprintIndexer::IndexAsset(const UObject* InAssetObject, FSearchSerializer& Serializer) const
 {
 	const UBlueprint* BP = Cast<UBlueprint>(InAssetObject);
 	check(BP);
 
-	//UGameplayAbilityBlueprint
-
 	if (UClass* GeneratedClass = BP->GeneratedClass)
 	{
 		if (UObject* CDO = GeneratedClass->GetDefaultObject())
@@ -52,10 +52,14 @@ void FBlueprintIndexer::IndexAsset(const UObject* InAssetObject, FSearchSerializ
 	{	
 		for (UEdGraphNode* Node : Graph->Nodes)
 		{
-			FText NodeType = LOCTEXT("Node", "Node");
-			FText NodeText = Node->GetNodeTitle(ENodeTitleType::MenuTitle);
+			const FText NodeText = Node->GetNodeTitle(ENodeTitleType::MenuTitle);
 			Serializer.BeginIndexingObject(Node, NodeText);
 			Serializer.IndexProperty(TEXT("Name"), NodeText);
+
+			if (!Node->NodeComment.IsEmpty())
+			{
+				Serializer.IndexProperty(TEXT("Comment"), Node->NodeComment);
+			}
 			
 			if (UK2Node_CallFunction* FunctionNode = Cast<UK2Node_CallFunction>(Node))
 			{
@@ -90,7 +94,7 @@ void FBlueprintIndexer::IndexAsset(const UObject* InAssetObject, FSearchSerializ
 	}
 }
 
-void FBlueprintIndexer::IndexMemberReference(FSearchSerializer& Serializer, const FMemberReference& MemberReference, const FString& MemberType)
+void FBlueprintIndexer::IndexMemberReference(FSearchSerializer& Serializer, const FMemberReference& MemberReference, const FString& MemberType) const
 {
 	Serializer.IndexProperty(MemberType + TEXT("Name"), MemberReference.GetMemberName());
 
diff --git a/Engine/Plugins/Editor/AssetSearch/Source/Private/Indexers/BlueprintIndexer.h b/Engine/Plugins/Editor/AssetSearch/Source/Private/Indexers/BlueprintIndexer.h
--- a/Engine/Plugins/Editor/AssetSearch/Source/Private/Indexers/BlueprintIndexer.h
+++ b/Engine/Plugins/Editor/AssetSearch/Source/Private/Indexers/BlueprintIndexer.h
@@ -12,8 +12,8 @@ class FBlueprintIndexer : public IAssetIndexer
 public:
 	virtual FString GetName() const override { return TEXT("Blueprint"); }
 	virtual int32 GetVersion() const override;
-	virtual void IndexAsset(const UObject* InAssetObject, FSearchSerializer& Serializer) override;
+	virtual void IndexAsset(const UObject* InAssetObject, FSearchSerializer& Serializer) const override;
 
 private:
-	void IndexMemberReference(FSearchSerializer& Serializer, const FMemberReference& MemberReference, const FString& MemberType);
+	void IndexMemberReference(FSearchSerializer& Serializer, const FMemberReference& MemberReference, const FString& MemberType) const;
 };
\ No newline at end of file
diff --git a/Engine/Plugins/Editor/AssetSearch/Source/Private/Indexers/DataAssetIndexer.cpp b/Engine/Plugins/Editor/AssetSearch/Source/Private/Indexers/DataAssetIndexer.cpp
--- a/Engine/Plugins/Editor/AssetSearch/Source/Private/Indexers/DataAssetIndexer.cpp
+++ b/Engine/Plugins/Editor/AssetSearch/Source/Private/Indexers/DataAssetIndexer.cpp
@@ -9,15 +9,17 @@ enum class EDataAssetIndexerVersion
 	Empty = 0,
 	Initial = 1,
 
-	Current = Initial,
+	// -----<new versions can be added above this line>-------------------------------------------------
+	VersionPlusOne,
+	LatestVersion = VersionPlusOne - 1
 };
 
 int32 FDataAssetIndexer::GetVersion() const
 {
-	return (int32)EDataAssetIndexerVersion::Current;
+	return (int32)EDataAssetIndexerVersion::LatestVersion;
 }
 
-void FDataAssetIndexer::IndexAsset(const UObject* InAssetObject, FSearchSerializer& Serializer)
+void FDataAssetIndexer::IndexAsset(const UObject* InAssetObject, FSearchSerializer& Serializer) const
 {
 	Serializer.BeginIndexingObject(InAssetObject, TEXT("$self"));
 
diff --git a/Engine/Plugins/Editor/AssetSearch/Source/Private/Indexers/DataAssetIndexer.h b/Engine/Plugins/Editor/AssetSearch/Source/Private/Indexers/DataAssetIndexer.h
--- a/Engine/Plugins/Editor/AssetSearch/Source/Private/Indexers/DataAssetIndexer.h
+++ b/Engine/Plugins/Editor/AssetSearch/Source/Private/Indexers/DataAssetIndexer.h
@@ -9,5 +9,5 @@ class FDataAssetIndexer : public IAssetIndexer
 {
 	virtual FString GetName() const override { return TEXT("DataAsset"); }
 	virtual int32 GetVersion() const override;
-	virtual void IndexAsset(const UObject* InAssetObject, FSearchSerializer& Serializer) override;
+	virtual void IndexAsset(const UObject* InAssetObject, FSearchSerializer& Serializer) const override;
 };
\ No newline at end of file
diff --git a/Engine/Plugins/Editor/AssetSearch/Source/Private/Indexers/DataTableIndexer.cpp b/Engine/Plugins/Editor/AssetSearch/Source/Private/Indexers/DataTableIndexer.cpp
--- a/Engine/Plugins/Editor/AssetSearch/Source/Private/Indexers/DataTableIndexer.cpp
+++ b/Engine/Plugins/Editor/AssetSearch/Source/Private/Indexers/DataTableIndexer.cpp
@@ -9,15 +9,17 @@ enum class EDataTableIndexerVersion
 	Empty = 0,
 	Initial = 1,
 
-	Current = Initial,
+	// -----<new versions can be added above this line>-------------------------------------------------
+	VersionPlusOne,
+	LatestVersion = VersionPlusOne - 1
 };
 
 int32 FDataTableIndexer::GetVersion() const
 {
-	return (int32)EDataTableIndexerVersion::Current;
+	return (int32)EDataTableIndexerVersion::LatestVersion;
 }
 
-void FDataTableIndexer::IndexAsset(const UObject* InAssetObject, FSearchSerializer& Serializer)
+void FDataTableIndexer::IndexAsset(const UObject* InAssetObject, FSearchSerializer& Serializer) const
 {
 	const UDataTable* DataTable = Cast<UDataTable>(InAssetObject);
 	check(DataTable);
diff --git a/Engine/Plugins/Editor/AssetSearch/Source/Private/Indexers/DataTableIndexer.h b/Engine/Plugins/Editor/AssetSearch/Source/Private/Indexers/DataTableIndexer.h
--- a/Engine/Plugins/Editor/AssetSearch/Source/Private/Indexers/DataTableIndexer.h
+++ b/Engine/Plugins/Editor/AssetSearch/Source/Private/Indexers/DataTableIndexer.h
@@ -9,5 +9,5 @@ class FDataTableIndexer : public IAssetIndexer
 {
 	virtual FString GetName() const override { return TEXT("DataTable"); }
 	virtual int32 GetVersion() const override;
-	virtual void IndexAsset(const UObject* InAssetObject, FSearchSerializer& Serializer) override;
+	virtual void IndexAsset(const UObject* InAssetObject, FSearchSerializer& Serializer) const override;
 };
\ No newline at end of file
diff --git a/Engine/Plugins/Editor/AssetSearch/Source/Private/Indexers/DialogueWaveIndexer.cpp b/Engine/Plugins/Editor/AssetSearch/Source/Private/Indexers/DialogueWaveIndexer.cpp
--- a/Engine/Plugins/Editor/AssetSearch/Source/Private/Indexers/DialogueWaveIndexer.cpp
+++ b/Engine/Plugins/Editor/AssetSearch/Source/Private/Indexers/DialogueWaveIndexer.cpp
@@ -10,15 +10,17 @@ enum class EDialogueWaveIndexerVersion
 	Empty = 0,
 	Initial = 1,
 
-	Current = Initial,
+	// -----<new versions can be added above this line>-------------------------------------------------
+	VersionPlusOne,
+	LatestVersion = VersionPlusOne - 1
 };
 
 int32 FDialogueWaveIndexer::GetVersion() const
 {
-	return (int32)EDialogueWaveIndexerVersion::Current;
+	return (int32)EDialogueWaveIndexerVersion::LatestVersion;
 }
 
-void FDialogueWaveIndexer::IndexAsset(const UObject* InAssetObject, FSearchSerializer& Serializer)
+void FDialogueWaveIndexer::IndexAsset(const UObject* InAssetObject, FSearchSerializer& Serializer) const
 {
 	const UDialogueWave* DialogueWave = Cast<UDialogueWave>(InAssetObject);
 	check(DialogueWave);
diff --git a/Engine/Plugins/Editor/AssetSearch/Source/Private/Indexers/DialogueWaveIndexer.h b/Engine/Plugins/Editor/AssetSearch/Source/Private/Indexers/DialogueWaveIndexer.h
--- a/Engine/Plugins/Editor/AssetSearch/Source/Private/Indexers/DialogueWaveIndexer.h
+++ b/Engine/Plugins/Editor/AssetSearch/Source/Private/Indexers/DialogueWaveIndexer.h
@@ -9,5 +9,5 @@ class FDialogueWaveIndexer : public IAssetIndexer
 {
 	virtual FString GetName() const override { return TEXT("DialogueWave"); }
 	virtual int32 GetVersion() const override;
-	virtual void IndexAsset(const UObject* InAssetObject, FSearchSerializer& Serializer) override;
+	virtual void IndexAsset(const UObject* InAssetObject, FSearchSerializer& Serializer) const override;
 };
\ No newline at end of file
diff --git a/Engine/Plugins/Editor/AssetSearch/Source/Private/Indexers/LevelIndexer.cpp b/Engine/Plugins/Editor/AssetSearch/Source/Private/Indexers/LevelIndexer.cpp
new file mode 100644
--- /dev/null
+++ b/Engine/Plugins/Editor/AssetSearch/Source/Private/Indexers/LevelIndexer.cpp
@@ -0,0 +1,37 @@
+// Copyright Epic Games, Inc. All Rights Reserved.
+
+#include "LevelIndexer.h"
+#include "Utility/IndexerUtilities.h"
+#include "Engine/DataAsset.h"
+#include "Engine/LevelScriptBlueprint.h"
+
+enum class ELevelIndexerVersion
+{
+	Empty = 0,
+	Initial = 1,
+
+	// -----<new versions can be added above this line>-------------------------------------------------
+	VersionPlusOne,
+	LatestVersion = VersionPlusOne - 1
+};
+
+int32 FLevelIndexer::GetVersion() const
+{
+	return (int32)ELevelIndexerVersion::LatestVersion;
+}
+
+void FLevelIndexer::GetNestedAssetTypes(TArray<UClass*>& OutTypes) const
+{
+	OutTypes.Add(UBlueprint::StaticClass());
+}
+
+void FLevelIndexer::IndexAsset(const UObject* InAssetObject, FSearchSerializer& Serializer) const
+{
+	const UWorld* World = Cast<UWorld>(InAssetObject);
+	check(World);
+
+	if (const ULevel* Level = World->PersistentLevel)
+	{
+		Serializer.IndexNestedAsset(const_cast<ULevel*>(Level)->GetLevelScriptBlueprint(true));
+	}
+}
\ No newline at end of file
diff --git a/Engine/Plugins/Editor/AssetSearch/Source/Private/Indexers/LevelIndexer.h b/Engine/Plugins/Editor/AssetSearch/Source/Private/Indexers/LevelIndexer.h
new file mode 100644
--- /dev/null
+++ b/Engine/Plugins/Editor/AssetSearch/Source/Private/Indexers/LevelIndexer.h
@@ -0,0 +1,14 @@
+// Copyright Epic Games, Inc. All Rights Reserved.
+
+#pragma once
+
+#include "CoreMinimal.h"
+#include "IAssetSearchModule.h"
+
+class FLevelIndexer : public IAssetIndexer
+{
+	virtual FString GetName() const override { return TEXT("Level"); }
+	virtual int32 GetVersion() const override;
+	virtual void IndexAsset(const UObject* InAssetObject, FSearchSerializer& Serializer) const override;
+	virtual void GetNestedAssetTypes(TArray<UClass*>& OutTypes) const override;
+};
\ No newline at end of file
diff --git a/Engine/Plugins/Editor/AssetSearch/Source/Private/Indexers/SoundCueIndexer.cpp b/Engine/Plugins/Editor/AssetSearch/Source/Private/Indexers/SoundCueIndexer.cpp
new file mode 100644
--- /dev/null
+++ b/Engine/Plugins/Editor/AssetSearch/Source/Private/Indexers/SoundCueIndexer.cpp
@@ -0,0 +1,46 @@
+// Copyright Epic Games, Inc. All Rights Reserved.
+
+#include "SoundCueIndexer.h"
+#include "Utility/IndexerUtilities.h"
+#include "Sound/SoundCue.h"
+#include "EdGraph/EdGraphNode.h"
+#include "Sound/SoundNode.h"
+
+enum class ESoundCueIndexerVersion
+{
+	Empty = 0,
+	Initial = 1,
+
+	// -----<new versions can be added above this line>-------------------------------------------------
+	VersionPlusOne,
+	LatestVersion = VersionPlusOne - 1
+};
+
+int32 FSoundCueIndexer::GetVersion() const
+{
+	return (int32)ESoundCueIndexerVersion::LatestVersion;
+}
+
+void FSoundCueIndexer::IndexAsset(const UObject* InAssetObject, FSearchSerializer& Serializer) const
+{
+	const USoundCue* SoundCue = Cast<USoundCue>(InAssetObject);
+	check(SoundCue);
+
+	for (USoundNode* SoundNode : SoundCue->AllNodes)
+	{
+		// Apparently SoundNodes can be null...
+		if (SoundNode)
+		{
+			if (UEdGraphNode* GraphNode = SoundNode->GetGraphNode())
+			{
+				const FText NodeText = GraphNode->GetNodeTitle(ENodeTitleType::MenuTitle);
+				Serializer.BeginIndexingObject(SoundNode, NodeText);
+				Serializer.IndexProperty(TEXT("Name"), NodeText);
+				FIndexerUtilities::IterateIndexableProperties(SoundNode, [&Serializer](const FProperty* Property, const FString& Value) {
+					Serializer.IndexProperty(Property, Value);
+				});
+				Serializer.EndIndexingObject();
+			}
+		}
+	}
+}
\ No newline at end of file
diff --git a/Engine/Plugins/Editor/AssetSearch/Source/Private/Indexers/SoundCueIndexer.h b/Engine/Plugins/Editor/AssetSearch/Source/Private/Indexers/SoundCueIndexer.h
new file mode 100644
--- /dev/null
+++ b/Engine/Plugins/Editor/AssetSearch/Source/Private/Indexers/SoundCueIndexer.h
@@ -0,0 +1,13 @@
+// Copyright Epic Games, Inc. All Rights Reserved.
+
+#pragma once
+
+#include "CoreMinimal.h"
+#include "IAssetSearchModule.h"
+
+class FSoundCueIndexer : public IAssetIndexer
+{
+	virtual FString GetName() const override { return TEXT("SoundCue"); }
+	virtual int32 GetVersion() const override;
+	virtual void IndexAsset(const UObject* InAssetObject, FSearchSerializer& Serializer) const override;
+};
\ No newline at end of file
diff --git a/Engine/Plugins/Editor/AssetSearch/Source/Private/Indexers/WidgetBlueprintIndexer.cpp b/Engine/Plugins/Editor/AssetSearch/Source/Private/Indexers/WidgetBlueprintIndexer.cpp
--- a/Engine/Plugins/Editor/AssetSearch/Source/Private/Indexers/WidgetBlueprintIndexer.cpp
+++ b/Engine/Plugins/Editor/AssetSearch/Source/Private/Indexers/WidgetBlueprintIndexer.cpp
@@ -13,15 +13,17 @@ enum class EWidgetBlueprintIndexerVersion
 	Empty = 0,
 	Initial = 1,
 
-	Current = Initial,
+	// -----<new versions can be added above this line>-------------------------------------------------
+	VersionPlusOne,
+	LatestVersion = VersionPlusOne - 1
 };
 
 int32 FWidgetBlueprintIndexer::GetVersion() const
 {
-	return (int32)EWidgetBlueprintIndexerVersion::Current;
+	return (int32)EWidgetBlueprintIndexerVersion::LatestVersion;
 }
 
-void FWidgetBlueprintIndexer::IndexAsset(const UObject* InAssetObject, FSearchSerializer& Serializer)
+void FWidgetBlueprintIndexer::IndexAsset(const UObject* InAssetObject, FSearchSerializer& Serializer) const
 {
 	const UWidgetBlueprint* BP = Cast<UWidgetBlueprint>(InAssetObject);
 	check(BP);
diff --git a/Engine/Plugins/Editor/AssetSearch/Source/Private/Indexers/WidgetBlueprintIndexer.h b/Engine/Plugins/Editor/AssetSearch/Source/Private/Indexers/WidgetBlueprintIndexer.h
--- a/Engine/Plugins/Editor/AssetSearch/Source/Private/Indexers/WidgetBlueprintIndexer.h
+++ b/Engine/Plugins/Editor/AssetSearch/Source/Private/Indexers/WidgetBlueprintIndexer.h
@@ -9,5 +9,5 @@ class FWidgetBlueprintIndexer : public IAssetIndexer
 {
 	virtual FString GetName() const override { return TEXT("WidgetBlueprint"); }
 	virtual int32 GetVersion() const override;
-	virtual void IndexAsset(const UObject* InAssetObject, FSearchSerializer& Serializer) override;
+	virtual void IndexAsset(const UObject* InAssetObject, FSearchSerializer& Serializer) const override;
 };
\ No newline at end of file
diff --git a/Engine/Plugins/Editor/AssetSearch/Source/Private/SearchSerializer.cpp b/Engine/Plugins/Editor/AssetSearch/Source/Private/SearchSerializer.cpp
--- a/Engine/Plugins/Editor/AssetSearch/Source/Private/SearchSerializer.cpp
+++ b/Engine/Plugins/Editor/AssetSearch/Source/Private/SearchSerializer.cpp
@@ -16,7 +16,9 @@ enum class ESearchIndexVersion
 	Empty = 0,
 	Initial = 1,
 
-	Current = Initial,
+	// -----<new versions can be added above this line>-------------------------------------------------
+	VersionPlusOne,
+	LatestVersion = VersionPlusOne - 1
 };
 
 static FName StaticGetNativeClassName(const UClass* InClass)
@@ -65,11 +67,50 @@ FSearchSerializer::~FSearchSerializer()
 
 int32 FSearchSerializer::GetVersion()
 {
-	return (int32)ESearchIndexVersion::Current;
+	return (int32)ESearchIndexVersion::LatestVersion;
 }
 
-void FSearchSerializer::BeginIndexer(IAssetIndexer* InIndexer)
+bool FSearchSerializer::IndexAsset(UObject* InAsset, const TMap<FName, TUniquePtr<IAssetIndexer>>& Indexers)
 {
+	if (InAsset == nullptr)
+	{
+		return false;
+	}
+
+	bool bWasIndexed = false;
+
+	UClass* IndexableClass = InAsset->GetClass();
+	while (IndexableClass)
+	{
+		if (const TUniquePtr<IAssetIndexer>* IndexerPtr = Indexers.Find(IndexableClass->GetFName()))
+		{
+			IAssetIndexer* Indexer = IndexerPtr->Get();
+
+			bWasIndexed = true;
+			BeginIndexer(Indexer);
+			Indexer->IndexAsset(InAsset, *this);
+			EndIndexer();
+		}
+
+		IndexableClass = IndexableClass->GetSuperClass();
+	}
+
+	// 
+	TArray<UObject*> TempNestedAssets = NestedAssets;
+	NestedAssets.Reset();
+	for (UObject* NestedAsset : TempNestedAssets)
+	{
+		IndexAsset(NestedAsset, Indexers);
+	}
+
+	return true;
+}
+
+void FSearchSerializer::BeginIndexer(const IAssetIndexer* InIndexer)
+{
+	check(CurrentIndexer == nullptr);
+	CurrentIndexer = InIndexer;
+
 	JsonWriter->WriteObjectStart(InIndexer->GetName());
 	JsonWriter->WriteValue(TEXT("version"), InIndexer->GetVersion());
 
@@ -80,6 +121,8 @@ void FSearchSerializer::EndIndexer()
 {
 	JsonWriter->WriteArrayEnd();
 	JsonWriter->WriteObjectEnd();
+
+	CurrentIndexer = nullptr;
 }
 
 void FSearchSerializer::BeginIndexingObject(const UObject* InObject, const FText& InFriendlyName)
@@ -98,6 +141,11 @@ void FSearchSerializer::BeginIndexingObject(const UObject* InObject, const FStri
 	CurrentObjectPath = InObject->GetPathName();
 }
 
+void FSearchSerializer::IndexNestedAsset(UObject* InNestedAsset)
+{
+	NestedAssets.Add(InNestedAsset);
+}
+
 void FSearchSerializer::EndIndexingObject()
 {
 	if (Values.Num() > 0)
diff --git a/Engine/Plugins/Editor/AssetSearch/Source/Private/Settings/AssetSearchDeveloperSettings.cpp b/Engine/Plugins/Editor/AssetSearch/Source/Private/Settings/AssetSearchDeveloperSettings.cpp
new file mode 100644
--- /dev/null
+++ b/Engine/Plugins/Editor/AssetSearch/Source/Private/Settings/AssetSearchDeveloperSettings.cpp
@@ -0,0 +1,7 @@
+// Copyright Epic Games, Inc. All Rights Reserved.
+
+#include "AssetSearchDeveloperSettings.h"
+
+UAssetSearchDeveloperSettings::UAssetSearchDeveloperSettings()
+{
+}
diff --git a/Engine/Plugins/Editor/AssetSearch/Source/Private/Settings/AssetSearchDeveloperSettings.h b/Engine/Plugins/Editor/AssetSearch/Source/Private/Settings/AssetSearchDeveloperSettings.h
new file mode 100644
--- /dev/null
+++ b/Engine/Plugins/Editor/AssetSearch/Source/Private/Settings/AssetSearchDeveloperSettings.h
@@ -0,0 +1,18 @@
+// Copyright Epic Games, Inc. All Rights Reserved.
+
+#pragma once
+
+#include "Engine/DeveloperSettings.h"
+#include "AssetSearchDeveloperSettings.generated.h"
+
+UCLASS(config = Editor, defaultconfig, meta=(DisplayName="Asset Search"))
+class UAssetSearchDeveloperSettings : public UDeveloperSettings
+{
+	GENERATED_BODY()
+
+public:
+	UAssetSearchDeveloperSettings();
+
+	UPROPERTY(config, EditAnywhere, Category=General)
+	TArray<FDirectoryPath> IgnoredPaths;
+};
diff --git a/Engine/Plugins/Editor/AssetSearch/Source/Public/IAssetSearchModule.h b/Engine/Plugins/Editor/AssetSearch/Source/Public/IAssetSearchModule.h
--- a/Engine/Plugins/Editor/AssetSearch/Source/Public/IAssetSearchModule.h
+++ b/Engine/Plugins/Editor/AssetSearch/Source/Public/IAssetSearchModule.h
@@ -9,6 +9,8 @@
 #include "SearchSerializer.h"
 #include "Modules/ModuleManager.h"
 
+class UClass;
+
 struct FSearchQuery
 {
 	FString Query;
@@ -52,7 +54,13 @@ public:
 	virtual ~IAssetIndexer() = 0;
 	virtual FString GetName() const = 0;
 	virtual int32 GetVersion() const = 0;
-	virtual void IndexAsset(const UObject* InAssetObject, FSearchSerializer& Serializer) = 0;
+	virtual void IndexAsset(const UObject* InAssetObject, FSearchSerializer& Serializer) const = 0;
+
+	/**
+	 * If your package contains a nested asset, such as the Blueprint stored in Level/World packages, 
+	 * it would return UBlueprint's class in the array.  This is only important if you use IndexNestedAsset.
+	 */
+	virtual void GetNestedAssetTypes(TArray<UClass*>& OutTypes) const { }
 };
 
 inline IAssetIndexer::~IAssetIndexer() = default;
diff --git a/Engine/Plugins/Editor/AssetSearch/Source/Public/SearchSerializer.h b/Engine/Plugins/Editor/AssetSearch/Source/Public/SearchSerializer.h
--- a/Engine/Plugins/Editor/AssetSearch/Source/Public/SearchSerializer.h
+++ b/Engine/Plugins/Editor/AssetSearch/Source/Public/SearchSerializer.h
@@ -21,7 +21,9 @@ public:
 
 	static int32 GetVersion();
 
-	void BeginIndexer(IAssetIndexer* InIndexer);
+	bool IndexAsset(UObject* InAsset, const TMap<FName, TUniquePtr<IAssetIndexer>>& Indexers);
+
+	void BeginIndexer(const IAssetIndexer* InIndexer);
 
 	//---------------------------------------------------------------
 
@@ -40,6 +42,10 @@ public:
 
 	//---------------------------------------------------------------
 
+	ASSETSEARCH_API void IndexNestedAsset(UObject* InNestedAsset);
+
+	//---------------------------------------------------------------
+
 	void EndIndexer();
 
 private:
@@ -53,6 +59,8 @@ private:
 
 	FAssetData AssetData;
 
+	const IAssetIndexer* CurrentIndexer = nullptr;
+
 	const UObject* CurrentObject = nullptr;
 	FString CurrentObjectName;
 	FString CurrentObjectClass;
@@ -69,4 +77,6 @@ private:
 	};
 
 	TArray<FIndexedValue> Values;
+
+	TArray<UObject*> NestedAssets;
 };
