Fixing some nonunity build issues, iwyu.

diff --git a/Engine/Plugins/Editor/AssetSearch/Source/Private/AssetSearchDatabase.cpp b/Engine/Plugins/Editor/AssetSearch/Source/Private/AssetSearchDatabase.cpp
--- a/Engine/Plugins/Editor/AssetSearch/Source/Private/AssetSearchDatabase.cpp
+++ b/Engine/Plugins/Editor/AssetSearch/Source/Private/AssetSearchDatabase.cpp
@@ -13,6 +13,8 @@
 #include "Misc/PackageName.h"
 #include "Misc/Paths.h"
 #include "Misc/SecureHash.h"
+#include "Engine/World.h"
+#include "HAL/PlatformFilemanager.h"
 
 PRAGMA_DISABLE_OPTIMIZATION
 
diff --git a/Engine/Plugins/Editor/AssetSearch/Source/Private/AssetSearchManager.cpp b/Engine/Plugins/Editor/AssetSearch/Source/Private/AssetSearchManager.cpp
--- a/Engine/Plugins/Editor/AssetSearch/Source/Private/AssetSearchManager.cpp
+++ b/Engine/Plugins/Editor/AssetSearch/Source/Private/AssetSearchManager.cpp
@@ -29,6 +29,8 @@
 #include "Indexers/SoundCueIndexer.h"
 #include "Sound/SoundCue.h"
 #include "Misc/ScopedSlowTask.h"
+#include "Misc/StringBuilder.h"
+#include "Engine/World.h"
 
 PRAGMA_DISABLE_OPTIMIZATION
 
diff --git a/Engine/Plugins/Editor/AssetSearch/Source/Private/Settings/AssetSearchDeveloperSettings.h b/Engine/Plugins/Editor/AssetSearch/Source/Private/Settings/AssetSearchDeveloperSettings.h
--- a/Engine/Plugins/Editor/AssetSearch/Source/Private/Settings/AssetSearchDeveloperSettings.h
+++ b/Engine/Plugins/Editor/AssetSearch/Source/Private/Settings/AssetSearchDeveloperSettings.h
@@ -3,6 +3,7 @@
 #pragma once
 
 #include "Engine/DeveloperSettings.h"
+#include "Engine/EngineTypes.h"
 #include "AssetSearchDeveloperSettings.generated.h"
 
 UCLASS(config = Editor, defaultconfig, meta=(DisplayName="Asset Search"))
