Fixing some includes for nonunity.

diff --git a/Engine/Plugins/Editor/AssetSearch/Source/Private/AssetSearchManager.cpp b/Engine/Plugins/Editor/AssetSearch/Source/Private/AssetSearchManager.cpp
--- a/Engine/Plugins/Editor/AssetSearch/Source/Private/AssetSearchManager.cpp
+++ b/Engine/Plugins/Editor/AssetSearch/Source/Private/AssetSearchManager.cpp
@@ -31,6 +31,7 @@
 #include "Misc/ScopedSlowTask.h"
 #include "Misc/StringBuilder.h"
 #include "Engine/World.h"
+#include "Editor.h"
 
 PRAGMA_DISABLE_OPTIMIZATION
 
diff --git a/Engine/Plugins/Editor/AssetSearch/Source/Private/Indexers/LevelIndexer.cpp b/Engine/Plugins/Editor/AssetSearch/Source/Private/Indexers/LevelIndexer.cpp
--- a/Engine/Plugins/Editor/AssetSearch/Source/Private/Indexers/LevelIndexer.cpp
+++ b/Engine/Plugins/Editor/AssetSearch/Source/Private/Indexers/LevelIndexer.cpp
@@ -4,6 +4,7 @@
 #include "Utility/IndexerUtilities.h"
 #include "Engine/DataAsset.h"
 #include "Engine/LevelScriptBlueprint.h"
+#include "Engine/World.h"
 
 enum class ELevelIndexerVersion
 {
