Fixing labels for widgets showing uo blank in the indexing.

diff --git a/Engine/Plugins/Editor/AssetSearch/Source/Private/Indexers/WidgetBlueprintIndexer.cpp b/Engine/Plugins/Editor/AssetSearch/Source/Private/Indexers/WidgetBlueprintIndexer.cpp
--- a/Engine/Plugins/Editor/AssetSearch/Source/Private/Indexers/WidgetBlueprintIndexer.cpp
+++ b/Engine/Plugins/Editor/AssetSearch/Source/Private/Indexers/WidgetBlueprintIndexer.cpp
@@ -12,6 +12,7 @@ enum class EWidgetBlueprintIndexerVersion
 {
 	Empty = 0,
 	Initial = 1,
+	FixLabels = 2,
 
 	// -----<new versions can be added above this line>-------------------------------------------------
 	VersionPlusOne,
@@ -32,8 +33,14 @@ void FWidgetBlueprintIndexer::IndexAsset(const UObject* InAssetObject, FSearchSe
 
 	for (const UWidget* Widget : AllWidgets)
 	{
-		Serializer.BeginIndexingObject(Widget, Widget->GetDisplayLabel());
-		Serializer.IndexProperty(TEXT("Name"), Widget->GetDisplayLabel());
+		FString Label = Widget->GetDisplayLabel();
+		if (Label.IsEmpty())
+		{
+			Label = Widget->GetName();
+		}
+
+		Serializer.BeginIndexingObject(Widget, Label);
+		Serializer.IndexProperty(TEXT("Name"), Label);
 
 		FIndexerUtilities::IterateIndexableProperties(Widget, [&Serializer](const FProperty* Property, const FString& Value) {
 			Serializer.IndexProperty(Property, Value);
