Taking the GT tasks off the GT task graph, and instead putting them onto a queue we can control, otherwise they can end up getting processed during saving packages.  Also fixing indexing pins for blueprints to actually index the value.

diff --git a/Engine/Plugins/Editor/AssetSearch/Source/Private/AssetSearchManager.cpp b/Engine/Plugins/Editor/AssetSearch/Source/Private/AssetSearchManager.cpp
--- a/Engine/Plugins/Editor/AssetSearch/Source/Private/AssetSearchManager.cpp
+++ b/Engine/Plugins/Editor/AssetSearch/Source/Private/AssetSearchManager.cpp
@@ -249,7 +249,7 @@ bool FAssetSearchManager::RequestIndexAsset(UObject* InAsset)
 				FScopeLock ScopedLock(&SearchDatabaseCS);
 				if (!SearchDatabase.IsAssetUpToDate(AssetData, InDDCKey))
 				{
-					AsyncTask(ENamedThreads::GameThread, [this, AssetWeakPtr]() {
+					AsyncMainThreadTask([this, AssetWeakPtr]() {
 						StoreIndexForAsset(AssetWeakPtr.Get());
 					});
 				}
@@ -340,7 +340,7 @@ bool FAssetSearchManager::AsyncGetDerivedDataKey(const FAssetData& InAssetData,
 			DDCKey.Append(LexToString(FileInfo.Hash));
 
 			const FString DDCKeyString = DDCKey.ToString();
-			AsyncTask(ENamedThreads::GameThread, [this, DDCKeyString, DDCKeyCallback]() {
+			AsyncMainThreadTask([this, DDCKeyString, DDCKeyCallback]() {
 				DDCKeyCallback(DDCKeyString);
 			});
 		}
@@ -596,12 +596,38 @@ void FAssetSearchManager::Search(const FSearchQuery& Query, TFunction<void(TArra
 			});
 		}
 
-		AsyncTask(ENamedThreads::GameThread, [ResultsFwd = MoveTemp(Results), InCallback]() mutable {
+		AsyncMainThreadTask([ResultsFwd = MoveTemp(Results), InCallback]() mutable {
 			InCallback(MoveTemp(ResultsFwd));
 		});
 	});
 }
 
+void FAssetSearchManager::AsyncMainThreadTask(TFunction<void()> Task)
+{
+	GT_Tasks.Enqueue(Task);
+
+	AsyncTask(ENamedThreads::GameThread, [this]() {
+		ProcessGameThreadTasks();
+	});
+}
+
+void FAssetSearchManager::ProcessGameThreadTasks()
+{
+	if (GIsSavingPackage)
+	{
+		AsyncTask(ENamedThreads::GameThread, [this]() {
+			ProcessGameThreadTasks();
+		});
+		return;
+	}
+
+	TFunction<void()> Operation;
+	if (GT_Tasks.Dequeue(Operation))
+	{
+		Operation();
+	}
+}
+
 #undef LOCTEXT_NAMESPACE
 
 PRAGMA_ENABLE_OPTIMIZATION
\ No newline at end of file
diff --git a/Engine/Plugins/Editor/AssetSearch/Source/Private/AssetSearchManager.h b/Engine/Plugins/Editor/AssetSearch/Source/Private/AssetSearchManager.h
--- a/Engine/Plugins/Editor/AssetSearch/Source/Private/AssetSearchManager.h
+++ b/Engine/Plugins/Editor/AssetSearch/Source/Private/AssetSearchManager.h
@@ -51,6 +51,9 @@ private:
 	void StoreIndexForAsset(UObject* InAsset);
 	void LoadDDCContentIntoDatabase(const FAssetData& InAsset, const TArray<uint8>& Content, const FString& DerivedDataKey);
 
+	void AsyncMainThreadTask(TFunction<void()> Task);
+	void ProcessGameThreadTasks();
+
 private:
 	FAssetSearchDatabase SearchDatabase;
 	FCriticalSection SearchDatabaseCS;
@@ -85,7 +88,7 @@ private:
 
 	FDelegateHandle TickerHandle;
 
-	TQueue<TFunction<void()>> GT_Requests;
+	TQueue<TFunction<void()>> GT_Tasks;
 
 private:
 	TAtomic<bool> RunThread;
diff --git a/Engine/Plugins/Editor/AssetSearch/Source/Private/Indexers/BlueprintIndexer.cpp b/Engine/Plugins/Editor/AssetSearch/Source/Private/Indexers/BlueprintIndexer.cpp
--- a/Engine/Plugins/Editor/AssetSearch/Source/Private/Indexers/BlueprintIndexer.cpp
+++ b/Engine/Plugins/Editor/AssetSearch/Source/Private/Indexers/BlueprintIndexer.cpp
@@ -8,6 +8,7 @@
 #include "Kismet2/BlueprintEditorUtils.h"
 #include "Utility/IndexerUtilities.h"
 #include "K2Node_BaseMCDelegate.h"
+#include "Internationalization/Text.h"
 
 #define LOCTEXT_NAMESPACE "FBlueprintIndexer"
 
@@ -15,8 +16,9 @@ PRAGMA_DISABLE_OPTIMIZATION
 
 enum class EBlueprintIndexerVersion
 {
-	Empty = 0,
-	Initial = 1,
+	Empty,
+	Initial,
+	FixingPinsToSaveValues,
 
 	// -----<new versions can be added above this line>-------------------------------------------------
 	VersionPlusOne,
@@ -79,13 +81,20 @@ void FBlueprintIndexer::IndexAsset(const UObject* InAssetObject, FSearchSerializ
 			{
 				for (UEdGraphPin* Pin : Node->GetAllPins())
 				{
-					FText PinText = Pin->GetDisplayName();
+					const FText PinText = Pin->GetDisplayName();
 					if (PinText.IsEmpty())
 					{
 						continue;
 					}
 
-					Serializer.IndexProperty(TEXT("Pin"), PinText);
+					const FString PinValue = Pin->DefaultValue;
+					if (PinValue.IsEmpty())
+					{
+						continue;
+					}
+
+					const FString PinLabel = TEXT("[Pin] ") + *FTextInspector::GetSourceString(PinText);
+					Serializer.IndexProperty(PinLabel, PinValue);
 				}
 			}
 
diff --git a/Engine/Source/Editor/EditorStyle/Private/SlateEditorStyle.cpp b/Engine/Source/Editor/EditorStyle/Private/SlateEditorStyle.cpp
--- a/Engine/Source/Editor/EditorStyle/Private/SlateEditorStyle.cpp
+++ b/Engine/Source/Editor/EditorStyle/Private/SlateEditorStyle.cpp
@@ -4677,10 +4677,12 @@ void FSlateEditorStyle::FStyle::SetupGraphEditorStyles()
 			);
 		}
 
-		Set("ClassIcon.K2Node_CallFunction", new IMAGE_BRUSH("Icons/icon_Blueprint_NewFunction_16x", Icon16x16));
-		Set("ClassIcon.K2Node_Variable", new IMAGE_BRUSH("Graph/Icons/Node", Icon22x22, FLinearColor::White));
-		Set("ClassIcon.K2Node_VariableGet", new IMAGE_BRUSH("Graph/Icons/FIB_VarGet", Icon22x22, FLinearColor::White));
-		Set("ClassIcon.K2Node_VariableSet", new IMAGE_BRUSH("Graph/Icons/FIB_VarSet", Icon22x22, FLinearColor::White));
+		Set("ClassIcon.K2Node_CallFunction",	new IMAGE_BRUSH("Icons/icon_Blueprint_NewFunction_16x", Icon16x16));
+		Set("ClassIcon.K2Node_CustomEvent",		new IMAGE_BRUSH("Icons/icon_Blueprint_Event_16x", Icon16x16));
+		Set("ClassIcon.K2Node_Event",			new IMAGE_BRUSH("Icons/icon_Blueprint_Event_16x", Icon16x16));
+		Set("ClassIcon.K2Node_Variable",		new IMAGE_BRUSH("Graph/Icons/Node", Icon16x16, FLinearColor::White));
+		Set("ClassIcon.K2Node_VariableGet",		new IMAGE_BRUSH("Graph/Icons/FIB_VarGet", Icon16x16, FLinearColor::White));
+		Set("ClassIcon.K2Node_VariableSet",		new IMAGE_BRUSH("Graph/Icons/FIB_VarSet", Icon16x16, FLinearColor::White));
 
 
 		Set( "GraphEditor.Default_16x", new IMAGE_BRUSH("Icons/icon_Blueprint_Node_16x", Icon16x16));
