Avoiding the GT task graph tasks in the search manager from spawning new GT tasks perputually when saving, this creates a hang when saving sometimes.  Now if it fails it just tries again later on the game thread tick rather than reissuing a GT task.

diff --git a/Engine/Plugins/Editor/AssetSearch/Source/Private/AssetSearchManager.cpp b/Engine/Plugins/Editor/AssetSearch/Source/Private/AssetSearchManager.cpp
--- a/Engine/Plugins/Editor/AssetSearch/Source/Private/AssetSearchManager.cpp
+++ b/Engine/Plugins/Editor/AssetSearch/Source/Private/AssetSearchManager.cpp
@@ -461,6 +461,8 @@ bool FAssetSearchManager::Tick_GameThread(float DeltaTime)
 	//	}
 	//}
 
+	ProcessGameThreadTasks();
+
 	int32 ScanLimit = GameThread_AssetScanLimit;
 	while (ProcessAssetQueue.Num() > 0 && ScanLimit > 0 && PendingDownloads < PendingDownloadsMax)
 	{
@@ -613,18 +615,19 @@ void FAssetSearchManager::AsyncMainThreadTask(TFunction<void()> Task)
 
 void FAssetSearchManager::ProcessGameThreadTasks()
 {
-	if (GIsSavingPackage)
+	if (!GT_Tasks.IsEmpty())
 	{
-		AsyncTask(ENamedThreads::GameThread, [this]() {
-			ProcessGameThreadTasks();
-		});
-		return;
-	}
+		if (GIsSavingPackage)
+		{
+			// If we're saving packages just give up, the call in Tick_GameThread will do this later.
+			return;
+		}
 
-	TFunction<void()> Operation;
-	if (GT_Tasks.Dequeue(Operation))
-	{
-		Operation();
+		TFunction<void()> Operation;
+		if (GT_Tasks.Dequeue(Operation))
+		{
+			Operation();
+		}
 	}
 }
 
