Search - now waiting until we successfully open the database (retry every 30 seconds) before processing assets.  Prevents a crash when one editor boots for the same project while another instance is still shutting down.

Blueprint Indexer - We now index user editable properties exposed to the details panel on nodes, so for example, Switch On Name nodes, the individual names for each name are now indexed.

Search - When searching we now attempt to find runs of tokens to generate additional search terms.  For example, searching for Get Health, will now also search for GetHealth, however "Get Health" will be left as is with no additional search.

diff --git a/Engine/Plugins/Editor/AssetSearch/Source/Private/AssetSearchDatabase.cpp b/Engine/Plugins/Editor/AssetSearch/Source/Private/AssetSearchDatabase.cpp
--- a/Engine/Plugins/Editor/AssetSearch/Source/Private/AssetSearchDatabase.cpp
+++ b/Engine/Plugins/Editor/AssetSearch/Source/Private/AssetSearchDatabase.cpp
@@ -424,76 +424,7 @@ public:
 	FSearchAssetsFTS Statement_SearchAssetsFTS;
 	bool SearchAssets(const FSearchQuery& Query, TFunctionRef<ESQLitePreparedStatementExecuteRowResult(FSearchRecord&&)> InCallback)
 	{
-		//FString BaseSearch(
-		//	" SELECT "
-		//	"     asset_name, "
-		//	"     asset_class, "
-		//	"     asset_path, "
-		//	"     object_name, "
-		//	"     object_path, "
-		//	"     object_native_class, "
-		//	"     property_name, "
-		//	"     property_field, "
-		//	"     property_class, "
-		//	"     value_text, "
-		//	"     value_hidden "
-		//	" FROM view_asset_properties "
-		//	" WHERE ");
-
-		//const bool bEvaledFilter = Eval.SetFilterText(FText::FromString(Query.Query));
-		//if (!ensure(bEvaledFilter))
-		//{
-		//	return false;
-		//}
-
-		//FTextFilterExpressionEvaluator Eval(ETextFilterExpressionEvaluatorMode::BasicString);
-		//const TArray<FExpressionToken>& Tokens = Eval.GetFilterExpressionTokens();
-
-		//for (int32 i = 0; i < Tokens.Num(); i++)
-		//{
-		//	const FStringToken& Token = Tokens[i].Context;
-		//	
-		//	if (Tokens[i].Node.Cast<TextFilterExpressionParser::FEqual>())
-		//	{
-
-		//	}
-		//}
-
-
-		//FSQLitePreparedStatement DynamicSearchStatement = Database.PrepareStatement(ESQLitePreparedStatementFlags::None);
-		
-
-		FString Q;
-
-		FTextFilterExpressionEvaluator Eval(ETextFilterExpressionEvaluatorMode::BasicString);
-		const TArray<FExpressionToken>& Tokens = Eval.GetFilterExpressionTokens();
-		if (Eval.SetFilterText(FText::FromString(Query.Query)))
-		{
-			for (int32 i = 0; i < Tokens.Num(); i++)
-			{
-				const FStringToken& Token = Tokens[i].Context;
-
-				if (Token.GetString().StartsWith(TEXT("\"")) && Token.GetString().EndsWith(TEXT("\"")))
-				{
-					Q += Token.GetString();
-					Q += TEXT(" ");
-				}
-				else
-				{
-					Q += TEXT("\"") + Token.GetString() + TEXT("\" * ");
-				}
-			}
-		}
-		else
-		{
-			TArray<FString> Phrases;
-			Query.Query.ParseIntoArray(Phrases, TEXT(" "), 1);
-
-			for (FString Phrase : Phrases)
-			{
-				Q += TEXT("\"") + Phrase + TEXT("\" * ");
-			}
-		}
+		const FString Q = Query.ConvertToDatabaseQuery();
 
 		return Statement_SearchAssetsFTS.BindAndExecute(Q, [&InCallback](const FSearchAssetsFTS& InStatement)
 		{
@@ -976,4 +907,106 @@ void FAssetSearchDatabase::RemoveAssetsNotInThisSet(const TArray<FAssetData>& In
 	}
 }
 
+FString FSearchQuery::ConvertToDatabaseQuery() const
+{
+	TStringBuilder<512> Q;
+
+	FTextFilterExpressionEvaluator Eval(ETextFilterExpressionEvaluatorMode::BasicString);
+	const TArray<FExpressionToken>& Tokens = Eval.GetFilterExpressionTokens();
+	if (Eval.SetFilterText(FText::FromString(Query)))
+	{
+		TArray<FString> TokenStreak;
+		bool bBreakSteak = false;
+
+		for (int32 i = 0; i < Tokens.Num(); i++)
+		{
+			const FExpressionToken& Token = Tokens[i];
+			const FStringToken& TokenContext = Token.Context;
+			const FString TokenString = TokenContext.GetString();
+			
+			TStringBuilder<64> Phrase;
+			if (Token.Node.Cast<TextFilterExpressionParser::FTextToken>())
+			{
+				if (TokenString.StartsWith(TEXT("\"")) && TokenString.EndsWith(TEXT("\"")))
+				{
+					Phrase.Append(TokenString);
+					Phrase.Append(TEXT(" "));
+					bBreakSteak = true;
+				}
+				else
+				{
+					Phrase.Append(TEXT("\""));
+					Phrase.Append(TokenString);
+					Phrase.Append(TEXT("\" * "));
+
+					TokenStreak.Add(TokenString);
+				}
+			}
+			else if (Token.Node.Cast<TextFilterExpressionParser::FOr>())
+			{
+				Phrase.Append(TEXT(" OR "));
+				bBreakSteak = true;
+			}
+			else if (Token.Node.Cast<TextFilterExpressionParser::FAnd>())
+			{
+				//Phrase.Append(TEXT(" AND "));
+				//bBreakSteak = true;
+			}
+			else
+			{
+				bBreakSteak = true;
+			}
+
+			if (bBreakSteak)
+			{
+				if (TokenStreak.Num() > 1)
+				{
+					Q.Append(TEXT(" OR "));
+					Q.Append(TEXT("\""));
+					for (const FString SimpleString : TokenStreak)
+					{
+						Q.Append(SimpleString);
+					}
+					Q.Append(TEXT("\""));
+				}
+
+				bBreakSteak = false;
+				TokenStreak.Reset();
+			}
+
+			Q.Append(Phrase);
+		}
+
+		if (TokenStreak.Num() > 1)
+		{
+			Q.Append(TEXT(" OR "));
+			Q.Append(TEXT("\""));
+			for (const FString SimpleString : TokenStreak)
+			{
+				Q.Append(SimpleString);
+			}
+			Q.Append(TEXT("\""));
+		}
+	}
+	else
+	{
+		TArray<FString> Phrases;
+		Query.ParseIntoArray(Phrases, TEXT(" "), 1);
+
+		for (FString Phrase : Phrases)
+		{
+			Q.Append(TEXT("\""));
+			Q.Append(Phrase);
+			Q.Append(TEXT("\" * "));
+		}
+
+		Q.Append(TEXT(" OR "));
+		Q.Append(TEXT("\""));
+		Q.Append(Query.Replace(TEXT(" "), TEXT("")));
+		Q.Append(TEXT("\""));
+	}
+
+	return Q.ToString();
+}
+
 PRAGMA_ENABLE_OPTIMIZATION
diff --git a/Engine/Plugins/Editor/AssetSearch/Source/Private/AssetSearchManager.cpp b/Engine/Plugins/Editor/AssetSearch/Source/Private/AssetSearchManager.cpp
--- a/Engine/Plugins/Editor/AssetSearch/Source/Private/AssetSearchManager.cpp
+++ b/Engine/Plugins/Editor/AssetSearch/Source/Private/AssetSearchManager.cpp
@@ -102,10 +102,6 @@ void FAssetSearchManager::Start()
 	RegisterAssetIndexer(UWorld::StaticClass(), MakeUnique<FLevelIndexer>());
 	RegisterAssetIndexer(USoundCue::StaticClass(), MakeUnique<FSoundCueIndexer>());
 
-
-	const FString SessionPath = FPaths::ConvertRelativePathToFull(FPaths::Combine(FPaths::ProjectSavedDir(), TEXT("Search")));
-	SearchDatabase.Open(SessionPath);
-
 	FCoreUObjectDelegates::OnObjectSaved.AddRaw(this, &FAssetSearchManager::OnObjectSaved);
 	FCoreUObjectDelegates::OnAssetLoaded.AddRaw(this, &FAssetSearchManager::OnAssetLoaded);
 
@@ -128,6 +124,23 @@ void FAssetSearchManager::Start()
 	DatabaseThread = FRunnableThread::Create(this, TEXT("UniversalSearch"), 0, TPri_BelowNormal);
 }
 
+void FAssetSearchManager::TryConnectToDatabase()
+{
+	if (!bDatabaseOpen)
+	{
+		if ((FPlatformTime::Seconds() - LastConnectionAttempt) > 30)
+		{
+			LastConnectionAttempt = FPlatformTime::Seconds();
+
+			const FString SessionPath = FPaths::ConvertRelativePathToFull(FPaths::Combine(FPaths::ProjectSavedDir(), TEXT("Search")));
+			if (SearchDatabase.Open(SessionPath))
+			{
+				bDatabaseOpen = true;
+			}
+		}
+	}
+}
+
 FSearchStats FAssetSearchManager::GetStats() const
 {
 	FSearchStats Stats;
@@ -449,6 +462,8 @@ bool FAssetSearchManager::Tick_GameThread(float DeltaTime)
 {
 	check(IsInGameThread());
 
+	TryConnectToDatabase();
+
 	//if (0)
 	//{
 	//	IAssetRegistry& AssetRegistry = FModuleManager::LoadModuleChecked<FAssetRegistryModule>("AssetRegistry").Get();
@@ -539,6 +554,12 @@ void FAssetSearchManager::Tick_DatabaseOperationThread()
 {
 	while (RunThread)
 	{
+		if (!bDatabaseOpen)
+		{
+			FPlatformProcess::Sleep(1);
+			continue;
+		}
+
 		TFunction<void()> Operation;
 		if (ImmediateOperations.Dequeue(Operation) || FeedOperations.Dequeue(Operation) || UpdateOperations.Dequeue(Operation))
 		{
diff --git a/Engine/Plugins/Editor/AssetSearch/Source/Private/AssetSearchManager.h b/Engine/Plugins/Editor/AssetSearch/Source/Private/AssetSearchManager.h
--- a/Engine/Plugins/Editor/AssetSearch/Source/Private/AssetSearchManager.h
+++ b/Engine/Plugins/Editor/AssetSearch/Source/Private/AssetSearchManager.h
@@ -28,6 +28,8 @@ public:
 	void ForceIndexOnAssetsMissingIndex();
 
 private:
+	void TryConnectToDatabase();
+
 	bool Tick_GameThread(float DeltaTime);
 	virtual uint32 Run() override;
 	void Tick_DatabaseOperationThread();
@@ -91,6 +93,9 @@ private:
 	TQueue<TFunction<void()>> GT_Tasks;
 
 private:
+	bool bDatabaseOpen = false;
+	double LastConnectionAttempt = 0;
+
 	TAtomic<bool> RunThread;
 	FRunnableThread* DatabaseThread = nullptr;
 
diff --git a/Engine/Plugins/Editor/AssetSearch/Source/Private/Indexers/BlueprintIndexer.cpp b/Engine/Plugins/Editor/AssetSearch/Source/Private/Indexers/BlueprintIndexer.cpp
--- a/Engine/Plugins/Editor/AssetSearch/Source/Private/Indexers/BlueprintIndexer.cpp
+++ b/Engine/Plugins/Editor/AssetSearch/Source/Private/Indexers/BlueprintIndexer.cpp
@@ -19,6 +19,7 @@ enum class EBlueprintIndexerVersion
 	Empty,
 	Initial,
 	FixingPinsToSaveValues,
+	IndexingPublicEditableFieldsOnNodes,
 
 	// -----<new versions can be added above this line>-------------------------------------------------
 	VersionPlusOne,
@@ -98,6 +99,11 @@ void FBlueprintIndexer::IndexAsset(const UObject* InAssetObject, FSearchSerializ
 				}
 			}
 
+			// This will serialize any user exposed options for the node that are editable in the details panel.
+			FIndexerUtilities::IterateIndexableProperties(Node, [&Serializer](const FProperty* Property, const FString& Value) {
+				Serializer.IndexProperty(Property, Value);
+			});
+
 			Serializer.EndIndexingObject();
 		}
 	}
diff --git a/Engine/Plugins/Editor/AssetSearch/Source/Public/IAssetSearchModule.h b/Engine/Plugins/Editor/AssetSearch/Source/Public/IAssetSearchModule.h
--- a/Engine/Plugins/Editor/AssetSearch/Source/Public/IAssetSearchModule.h
+++ b/Engine/Plugins/Editor/AssetSearch/Source/Public/IAssetSearchModule.h
@@ -14,6 +14,8 @@ class UClass;
 struct FSearchQuery
 {
 	FString Query;
+
+	FString ConvertToDatabaseQuery() const;
 };
 
 struct FSearchRecord
