Fixing IWYU issue.

diff --git a/Engine/Plugins/Editor/AssetSearch/Source/Private/AssetSearchDatabase.cpp b/Engine/Plugins/Editor/AssetSearch/Source/Private/AssetSearchDatabase.cpp
--- a/Engine/Plugins/Editor/AssetSearch/Source/Private/AssetSearchDatabase.cpp
+++ b/Engine/Plugins/Editor/AssetSearch/Source/Private/AssetSearchDatabase.cpp
@@ -15,6 +15,7 @@
 #include "Misc/SecureHash.h"
 #include "Engine/World.h"
 #include "HAL/PlatformFilemanager.h"
+#include "Misc/StringBuilder.h"
 
 PRAGMA_DISABLE_OPTIMIZATION
 
