Universal Search - Better search feedback in the window (throbber).  Fixing all the versioning enums to not use explicit values for each entry, caused a bug where the database version didn't bump.  This CL also bumps the database version, so will need to recreate the database for people, but also splitting out the file hashing database into its own database to avoid needing to rehash files in the future if the asset search database format changes.

diff --git a/Engine/Plugins/Editor/AssetSearch/Source/Private/AssetSearchDatabase.cpp b/Engine/Plugins/Editor/AssetSearch/Source/Private/AssetSearchDatabase.cpp
--- a/Engine/Plugins/Editor/AssetSearch/Source/Private/AssetSearchDatabase.cpp
+++ b/Engine/Plugins/Editor/AssetSearch/Source/Private/AssetSearchDatabase.cpp
@@ -23,10 +23,11 @@ DECLARE_LOG_CATEGORY_CLASS(LogAssetSearch, Log, All);
 
 enum class EAssetSearchDatabaseVersion
 {
-	Empty = 0,
-	Initial = 1,
-	IndexingAssetIdsAssetPathsUnique = 2,
-	IntroducingFileHashing = 2,
+	Empty,
+	Initial,
+	IndexingAssetIdsAssetPathsUnique,
+	IntroducingFileHashing,
+	MovingFileInfoIntoAnotherDatabase,
 
 	// -----<new versions can be added above this line>-------------------------------------------------
 	VersionPlusOne,
@@ -67,11 +68,6 @@ public:
 
 		PREPARE_STATEMENT(Statement_SearchAssetsFTS);
 
-		PREPARE_STATEMENT(Statement_AddFileInfo);
-		PREPARE_STATEMENT(Statement_UpdateFileInfo);
-		PREPARE_STATEMENT(Statement_GetFileInfo);
-		PREPARE_STATEMENT(Statement_GetAllFileInfos);
-
 #undef PREPARE_STATEMENT
 
 		return true;
@@ -160,109 +156,6 @@ public:
 		return INDEX_NONE;
 	}
 
-	struct FCachedFileInfo
-	{
-		int64 FileId;
-		FString FilePath;
-		FDateTime LastModifed;
-		FString Hash;
-
-		FAssetFileInfo ToAssetFileInfo() const
-		{
-			FAssetFileInfo AssetFileInfo;
-			AssetFileInfo.LastModified = LastModifed;
-			AssetFileInfo.PackageName = *FilePath;
-			LexFromString(AssetFileInfo.Hash, *Hash);
-
-			return AssetFileInfo;
-		}
-	};
-
-	SQLITE_PREPARED_STATEMENT(FGetFileInfo,
-		"SELECT fileid, file_last_modified, file_hash FROM table_files WHERE file_path = ?1;",
-		SQLITE_PREPARED_STATEMENT_COLUMNS(int64, FDateTime, FString),
-		SQLITE_PREPARED_STATEMENT_BINDINGS(FString)
-	);
-	private: FGetFileInfo Statement_GetFileInfo;
-	public: bool GetFileInfo(const FString& InFullFilePath, FCachedFileInfo& OutFileInfo)
-	{
-		OutFileInfo.FilePath = InFullFilePath.ToLower();
-		if (Statement_GetFileInfo.BindAndExecuteSingle(OutFileInfo.FilePath, OutFileInfo.FileId, OutFileInfo.LastModifed, OutFileInfo.Hash))
-		{
-			return true;
-		}
-
-		return false;
-	}
-
-	SQLITE_PREPARED_STATEMENT(FGetAllFileInfos,
-		"SELECT file_path, file_last_modified, file_hash FROM table_files;",
-		SQLITE_PREPARED_STATEMENT_COLUMNS(FString, FDateTime, FString),
-		SQLITE_PREPARED_STATEMENT_BINDINGS()
-	);
-	FGetAllFileInfos Statement_GetAllFileInfos;
-	bool GetAllFileInfos(TFunctionRef<ESQLitePreparedStatementExecuteRowResult(FAssetFileInfo&&)> InCallback)
-	{
-		return Statement_GetAllFileInfos.BindAndExecute([&InCallback](const FGetAllFileInfos& InStatement)
-		{
-			FCachedFileInfo FileInfo;
-			if (InStatement.GetColumnValues(FileInfo.FilePath, FileInfo.LastModifed, FileInfo.Hash))
-			{
-				return InCallback(FileInfo.ToAssetFileInfo());
-			}
-			return ESQLitePreparedStatementExecuteRowResult::Error;
-		}) != INDEX_NONE;
-	}
-
-	SQLITE_PREPARED_STATEMENT_BINDINGS_ONLY(
-		FUpdateFileInfo,
-		" UPDATE table_files SET file_last_modified = ?2, file_hash = ?3 WHERE file_path = ?1;",
-		SQLITE_PREPARED_STATEMENT_BINDINGS(FString, FDateTime, FString)
-	);
-	private: FUpdateFileInfo Statement_UpdateFileInfo;
-	SQLITE_PREPARED_STATEMENT_BINDINGS_ONLY(
-		FAddFileInfo,
-		" INSERT INTO table_files(file_path, file_last_modified, file_hash)"
-		" VALUES(?1, ?2, ?3);",
-		SQLITE_PREPARED_STATEMENT_BINDINGS(FString, FDateTime, FString)
-	);
-	private: FAddFileInfo Statement_AddFileInfo;
-	public: bool AddOrUpdateFileInfo(const FAssetData& InAssetData, FAssetFileInfo& OutFileInfo)
-	{
-		const FString PackageName = InAssetData.PackageName.ToString();
-		const bool bIsWorldAsset = (InAssetData.AssetClass == UWorld::StaticClass()->GetFName());
-		const FString Extension = bIsWorldAsset ? FPackageName::GetMapPackageExtension() : FPackageName::GetAssetPackageExtension();
-		const FString FilePath = FPackageName::LongPackageNameToFilename(PackageName, Extension);
-		const FString FullFilePath = FPaths::ConvertRelativePathToFull(FilePath);
-
-		IPlatformFile& PlatformFile = FPlatformFileManager::Get().GetPlatformFile();
-		const FDateTime CurrentLastModified = PlatformFile.GetTimeStamp(*FullFilePath);
-
-		FCachedFileInfo FileInfo;
-		if (GetFileInfo(PackageName.ToLower(), FileInfo))
-		{
-			if (CurrentLastModified == FileInfo.LastModifed)
-			{
-				OutFileInfo = FileInfo.ToAssetFileInfo();
-				return false;
-			}
-
-			OutFileInfo = FileInfo.ToAssetFileInfo();
-			OutFileInfo.LastModified = CurrentLastModified;
-			OutFileInfo.Hash = FMD5Hash::HashFile(*FullFilePath);
-			Statement_UpdateFileInfo.BindAndExecuteSingle(PackageName.ToLower(), OutFileInfo.LastModified, LexToString(OutFileInfo.Hash));
-			return true;
-		}
-		else
-		{
-			OutFileInfo = FileInfo.ToAssetFileInfo();
-			OutFileInfo.LastModified = CurrentLastModified;
-			OutFileInfo.Hash = FMD5Hash::HashFile(*FullFilePath);
-			Statement_AddFileInfo.BindAndExecuteSingle(PackageName.ToLower(), OutFileInfo.LastModified, LexToString(OutFileInfo.Hash));
-			return true;
-		}
-	}
-
 	SQLITE_PREPARED_STATEMENT_BINDINGS_ONLY(
 		FAddAssetToAssetTable,
 		" INSERT INTO table_assets(asset_name, asset_class, asset_path, index_hash)"
@@ -580,13 +473,6 @@ bool FAssetSearchDatabase::Open(const FString& InSessionPath, const ESQLiteDatab
 
 	// Create our required tables
 	//========================================================================
-	if (!ensure(Database->Execute(TEXT("CREATE TABLE IF NOT EXISTS table_files(fileid INTEGER PRIMARY KEY, file_path TEXT UNIQUE, file_last_modified INTEGER NOT NULL, file_hash);"))))
-	{
-		LogLastError();
-		Close();
-		return false;
-	}
-
 	if (!ensure(Database->Execute(TEXT("CREATE TABLE IF NOT EXISTS table_assets(assetid INTEGER PRIMARY KEY, asset_name, asset_class, asset_path TEXT UNIQUE, index_hash);"))))
 	{
 		LogLastError();
@@ -691,15 +577,6 @@ bool FAssetSearchDatabase::Open(const FString& InSessionPath, const ESQLiteDatab
 		return false;
 	}
 
-	if (!ensure(Database->Execute(
-		TEXT("CREATE UNIQUE INDEX IF NOT EXISTS file_path_index ON table_files(file_path);")
-	)))
-	{
-		LogLastError();
-		Close();
-		return false;
-	}
-
 	if (!ensure(Database->Execute(
 		TEXT("CREATE UNIQUE INDEX IF NOT EXISTS asset_path_index ON table_assets(asset_path);")
 	)))
@@ -779,16 +656,6 @@ void FAssetSearchDatabase::LogLastError() const
 	UE_LOG(LogAssetSearch, Error, TEXT("Database Error: %s"), *SessionPath, *GetLastError());
 }
 
-bool FAssetSearchDatabase::AddOrUpdateFileInfo(const FAssetData& InAssetData, FAssetFileInfo& OutFileInfo)
-{
-	if (ensure(Statements))
-	{
-		return Statements->AddOrUpdateFileInfo(InAssetData, OutFileInfo);
-	}
-
-	return false;
-}
-
 bool FAssetSearchDatabase::IsAssetUpToDate(const FAssetData& InAssetData, const FString& IndexedJsonHash)
 {
 	if (ensure(Statements))
@@ -845,36 +712,6 @@ void FAssetSearchDatabase::RemoveAsset(const FAssetData& InAssetData)
 	}
 }
 
-void FAssetSearchDatabase::AddOrUpdateFileInfos(const TArray<FAssetData>& InAssets)
-{
-	for (const FAssetData& InAsset : InAssets)
-	{
-		// If it's a redirector act like it has been removed from the system,
-		// we don't want old duplicate entries for it.
-		if (InAsset.IsRedirector())
-		{
-			continue;
-		}
-
-		// Freshen hash cache
-		FAssetFileInfo FileInfo;
-		AddOrUpdateFileInfo(InAsset, FileInfo);
-	}
-}
-
-TMap<FName, FAssetFileInfo> FAssetSearchDatabase::GetAllFileInfos()
-{
-	TMap<FName, FAssetFileInfo> FileInfos;
-	Statements->GetAllFileInfos([&FileInfos](FAssetFileInfo&& InResult)
-	{
-		FileInfos.Add(InResult.PackageName, InResult);
-
-		return ESQLitePreparedStatementExecuteRowResult::Continue;
-	});
-
-	return FileInfos;
-}
-
 void FAssetSearchDatabase::RemoveAssetsNotInThisSet(const TArray<FAssetData>& InAssets)
 {
 	TSet<FString> InAssetPaths;
diff --git a/Engine/Plugins/Editor/AssetSearch/Source/Private/AssetSearchDatabase.h b/Engine/Plugins/Editor/AssetSearch/Source/Private/AssetSearchDatabase.h
--- a/Engine/Plugins/Editor/AssetSearch/Source/Private/AssetSearchDatabase.h
+++ b/Engine/Plugins/Editor/AssetSearch/Source/Private/AssetSearchDatabase.h
@@ -10,13 +10,6 @@ enum class ESQLiteDatabaseOpenMode : uint8;
 class FAssetSearchDatabaseStatements;
 struct FAssetData;
 
-struct FAssetFileInfo
-{
-	FName PackageName;
-	FDateTime LastModified;
-	FMD5Hash Hash;
-};
-
 class FAssetSearchDatabase
 {
 public:
@@ -71,16 +64,6 @@ public:
 	 */
 	FString GetLastError() const;
 
-	/**
-	 * Update the file hash table.
-	 */
-	bool AddOrUpdateFileInfo(const FAssetData& InAssetData, FAssetFileInfo& OutFileInfo);
-	void AddOrUpdateFileInfos(const TArray<FAssetData>& InAssets);
-
-	void UpdateFileHash(const FAssetData& InAssetData, FAssetFileInfo& OutFileInfo);
-
-	TMap<FName, FAssetFileInfo> GetAllFileInfos();
-
 	/**
 	 * Is the asset cache up to date for this asset?
 	 */
diff --git a/Engine/Plugins/Editor/AssetSearch/Source/Private/AssetSearchManager.cpp b/Engine/Plugins/Editor/AssetSearch/Source/Private/AssetSearchManager.cpp
--- a/Engine/Plugins/Editor/AssetSearch/Source/Private/AssetSearchManager.cpp
+++ b/Engine/Plugins/Editor/AssetSearch/Source/Private/AssetSearchManager.cpp
@@ -78,22 +78,27 @@ FAssetSearchManager::FAssetSearchManager()
 FAssetSearchManager::~FAssetSearchManager()
 {
 	RunThread = false;
-	DatabaseThread->WaitForCompletion();
 
+	if (DatabaseThread)
 	{
-		FScopeLock ScopedLock(&SearchDatabaseCS);
-		FCoreUObjectDelegates::OnObjectSaved.RemoveAll(this);
-		FCoreUObjectDelegates::OnAssetLoaded.RemoveAll(this);
-		UObject::FAssetRegistryTag::OnGetExtraObjectTags.RemoveAll(this);
+		DatabaseThread->WaitForCompletion();
+	}
 
-		//FModuleManager::GetModule<FAssetRegistryModule>("AssetRegistry");
+	FCoreUObjectDelegates::OnObjectSaved.RemoveAll(this);
+	FCoreUObjectDelegates::OnAssetLoaded.RemoveAll(this);
+	UObject::FAssetRegistryTag::OnGetExtraObjectTags.RemoveAll(this);
 
-		FTicker::GetCoreTicker().RemoveTicker(TickerHandle);
-	}
+	FTicker::GetCoreTicker().RemoveTicker(TickerHandle);
 }
 
 void FAssetSearchManager::Start()
 {
+	// Don't start the search manager on unattended builds.
+	//if (FApp::IsUnattended())
+	//{
+	//	return;
+	//}
+
 	RegisterAssetIndexer(UDataAsset::StaticClass(), MakeUnique<FDataAssetIndexer>());
 	RegisterAssetIndexer(UDataTable::StaticClass(), MakeUnique<FDataTableIndexer>());
 	RegisterAssetIndexer(UBlueprint::StaticClass(), MakeUnique<FBlueprintIndexer>());
@@ -133,10 +138,19 @@ void FAssetSearchManager::TryConnectToDatabase()
 			LastConnectionAttempt = FPlatformTime::Seconds();
 
 			const FString SessionPath = FPaths::ConvertRelativePathToFull(FPaths::Combine(FPaths::ProjectSavedDir(), TEXT("Search")));
-			if (SearchDatabase.Open(SessionPath))
+			
+			if (!FileInfoDatabase.Open(SessionPath))
 			{
-				bDatabaseOpen = true;
+				return;
 			}
+
+			if (!SearchDatabase.Open(SessionPath))
+			{
+				FileInfoDatabase.Close();
+				return;
+			}
+
+			bDatabaseOpen = true;
 		}
 	}
 }
@@ -327,10 +341,12 @@ bool FAssetSearchManager::AsyncGetDerivedDataKey(const FAssetData& InAssetData,
 	}
 
 	UpdateOperations.Enqueue([this, InAssetData, IndexersNamesAndVersions, DDCKeyCallback]() {
-		FScopeLock ScopedLock(&SearchDatabaseCS);
-
 		FAssetFileInfo FileInfo;
-		SearchDatabase.AddOrUpdateFileInfo(InAssetData, FileInfo);
+
+		{
+			FScopeLock ScopedLock(&FileInfoDatabaseCS);
+			FileInfoDatabase.AddOrUpdateFileInfo(InAssetData, FileInfo);
+		}
 
 		if (FileInfo.Hash.IsValid())
 		{
diff --git a/Engine/Plugins/Editor/AssetSearch/Source/Private/AssetSearchManager.h b/Engine/Plugins/Editor/AssetSearch/Source/Private/AssetSearchManager.h
--- a/Engine/Plugins/Editor/AssetSearch/Source/Private/AssetSearchManager.h
+++ b/Engine/Plugins/Editor/AssetSearch/Source/Private/AssetSearchManager.h
@@ -5,6 +5,7 @@
 #include "CoreMinimal.h"
 #include "IAssetSearchModule.h"
 #include "AssetSearchDatabase.h"
+#include "FileInfoDatabase.h"
 #include "Containers/Queue.h"
 #include "HAL/Runnable.h"
 
@@ -57,6 +58,8 @@ private:
 	void ProcessGameThreadTasks();
 
 private:
+	FFileInfoDatabase FileInfoDatabase;
+	FCriticalSection FileInfoDatabaseCS;
 	FAssetSearchDatabase SearchDatabase;
 	FCriticalSection SearchDatabaseCS;
 	TAtomic<int32> PendingDatabaseUpdates;
diff --git a/Engine/Plugins/Editor/AssetSearch/Source/Private/FileInfoDatabase.cpp b/Engine/Plugins/Editor/AssetSearch/Source/Private/FileInfoDatabase.cpp
new file mode 100644
--- /dev/null
+++ b/Engine/Plugins/Editor/AssetSearch/Source/Private/FileInfoDatabase.cpp
@@ -0,0 +1,436 @@
+// Copyright Epic Games, Inc. All Rights Reserved.
+
+#include "FileInfoDatabase.h"
+
+#include "SQLiteDatabase.h"
+#include "HAL/FileManager.h"
+#include "UObject/StructOnScope.h"
+#include "Serialization/MemoryReader.h"
+#include "Serialization/MemoryWriter.h"
+
+#include "AssetData.h"
+#include "Misc/TextFilterExpressionEvaluator.h"
+#include "Misc/PackageName.h"
+#include "Misc/Paths.h"
+#include "Misc/SecureHash.h"
+#include "Engine/World.h"
+#include "HAL/PlatformFilemanager.h"
+#include "Misc/StringBuilder.h"
+
+PRAGMA_DISABLE_OPTIMIZATION
+
+DECLARE_LOG_CATEGORY_CLASS(LogFileInfo, Log, All);
+
+enum class EFileInfoDatabaseVersion
+{
+	Empty,
+	Initial,
+
+	// -----<new versions can be added above this line>-------------------------------------------------
+	VersionPlusOne,
+	LatestVersion = VersionPlusOne - 1
+};
+
+class FFileInfoDatabaseStatements
+{
+private:
+	FSQLiteDatabase& Database;
+
+public:
+	explicit FFileInfoDatabaseStatements(FSQLiteDatabase& InDatabase)
+		: Database(InDatabase)
+	{
+		check(Database.IsValid());
+	}
+
+	bool CreatePreparedStatements()
+	{
+		check(Database.IsValid());
+
+#define PREPARE_STATEMENT(VAR)																				\
+			(VAR) = Database.PrepareStatement<decltype(VAR)>(ESQLitePreparedStatementFlags::Persistent);	\
+			if (!(VAR).IsValid()) { return false; }
+
+		PREPARE_STATEMENT(Statement_BeginTransaction);
+		PREPARE_STATEMENT(Statement_CommitTransaction);
+		PREPARE_STATEMENT(Statement_RollbackTransaction);
+
+		PREPARE_STATEMENT(Statement_AddFileInfo);
+		PREPARE_STATEMENT(Statement_UpdateFileInfo);
+		PREPARE_STATEMENT(Statement_GetFileInfo);
+		PREPARE_STATEMENT(Statement_GetAllFileInfos);
+
+#undef PREPARE_STATEMENT
+
+		return true;
+	}
+
+	/**
+	 * Statements managing database transactions
+	 */
+
+	 /** Begin a database transaction */
+	SQLITE_PREPARED_STATEMENT_SIMPLE(FBeginTransaction, "BEGIN TRANSACTION;");
+	FBeginTransaction Statement_BeginTransaction;
+	bool BeginTransaction()
+	{
+		return Statement_BeginTransaction.Execute();
+	}
+
+	/** Commit a database transaction */
+	SQLITE_PREPARED_STATEMENT_SIMPLE(FCommitTransaction, "COMMIT TRANSACTION;");
+	FCommitTransaction Statement_CommitTransaction;
+	bool CommitTransaction()
+	{
+		return Statement_CommitTransaction.Execute();
+	}
+
+	/** Rollback a database transaction */
+	SQLITE_PREPARED_STATEMENT_SIMPLE(FRollbackTransaction, "ROLLBACK TRANSACTION;");
+	FRollbackTransaction Statement_RollbackTransaction;
+	bool RollbackTransaction()
+	{
+		return Statement_RollbackTransaction.Execute();
+	}
+
+	/**
+	 * Application Statements
+	 */
+
+	struct FCachedFileInfo
+	{
+		int64 FileId;
+		FString FilePath;
+		FDateTime LastModifed;
+		FString Hash;
+
+		FAssetFileInfo ToAssetFileInfo() const
+		{
+			FAssetFileInfo AssetFileInfo;
+			AssetFileInfo.LastModified = LastModifed;
+			AssetFileInfo.PackageName = *FilePath;
+			LexFromString(AssetFileInfo.Hash, *Hash);
+
+			return AssetFileInfo;
+		}
+	};
+
+	SQLITE_PREPARED_STATEMENT(FGetFileInfo,
+		"SELECT fileid, file_last_modified, file_hash FROM table_files WHERE file_path = ?1;",
+		SQLITE_PREPARED_STATEMENT_COLUMNS(int64, FDateTime, FString),
+		SQLITE_PREPARED_STATEMENT_BINDINGS(FString)
+	);
+	private: FGetFileInfo Statement_GetFileInfo;
+	public: bool GetFileInfo(const FString& InFullFilePath, FCachedFileInfo& OutFileInfo)
+	{
+		OutFileInfo.FilePath = InFullFilePath.ToLower();
+		if (Statement_GetFileInfo.BindAndExecuteSingle(OutFileInfo.FilePath, OutFileInfo.FileId, OutFileInfo.LastModifed, OutFileInfo.Hash))
+		{
+			return true;
+		}
+
+		return false;
+	}
+
+	SQLITE_PREPARED_STATEMENT(FGetAllFileInfos,
+		"SELECT file_path, file_last_modified, file_hash FROM table_files;",
+		SQLITE_PREPARED_STATEMENT_COLUMNS(FString, FDateTime, FString),
+		SQLITE_PREPARED_STATEMENT_BINDINGS()
+	);
+	FGetAllFileInfos Statement_GetAllFileInfos;
+	bool GetAllFileInfos(TFunctionRef<ESQLitePreparedStatementExecuteRowResult(FAssetFileInfo&&)> InCallback)
+	{
+		return Statement_GetAllFileInfos.BindAndExecute([&InCallback](const FGetAllFileInfos& InStatement)
+		{
+			FCachedFileInfo FileInfo;
+			if (InStatement.GetColumnValues(FileInfo.FilePath, FileInfo.LastModifed, FileInfo.Hash))
+			{
+				return InCallback(FileInfo.ToAssetFileInfo());
+			}
+			return ESQLitePreparedStatementExecuteRowResult::Error;
+		}) != INDEX_NONE;
+	}
+
+	SQLITE_PREPARED_STATEMENT_BINDINGS_ONLY(
+		FUpdateFileInfo,
+		" UPDATE table_files SET file_last_modified = ?2, file_hash = ?3 WHERE file_path = ?1;",
+		SQLITE_PREPARED_STATEMENT_BINDINGS(FString, FDateTime, FString)
+	);
+	private: FUpdateFileInfo Statement_UpdateFileInfo;
+	SQLITE_PREPARED_STATEMENT_BINDINGS_ONLY(
+		FAddFileInfo,
+		" INSERT INTO table_files(file_path, file_last_modified, file_hash)"
+		" VALUES(?1, ?2, ?3);",
+		SQLITE_PREPARED_STATEMENT_BINDINGS(FString, FDateTime, FString)
+	);
+	private: FAddFileInfo Statement_AddFileInfo;
+	public: bool AddOrUpdateFileInfo(const FAssetData& InAssetData, FAssetFileInfo& OutFileInfo)
+	{
+		const FString PackageName = InAssetData.PackageName.ToString();
+		const bool bIsWorldAsset = (InAssetData.AssetClass == UWorld::StaticClass()->GetFName());
+		const FString Extension = bIsWorldAsset ? FPackageName::GetMapPackageExtension() : FPackageName::GetAssetPackageExtension();
+		const FString FilePath = FPackageName::LongPackageNameToFilename(PackageName, Extension);
+		const FString FullFilePath = FPaths::ConvertRelativePathToFull(FilePath);
+
+		IPlatformFile& PlatformFile = FPlatformFileManager::Get().GetPlatformFile();
+		const FDateTime CurrentLastModified = PlatformFile.GetTimeStamp(*FullFilePath);
+
+		FCachedFileInfo FileInfo;
+		if (GetFileInfo(PackageName.ToLower(), FileInfo))
+		{
+			if (CurrentLastModified == FileInfo.LastModifed)
+			{
+				OutFileInfo = FileInfo.ToAssetFileInfo();
+				return false;
+			}
+
+			OutFileInfo = FileInfo.ToAssetFileInfo();
+			OutFileInfo.LastModified = CurrentLastModified;
+			OutFileInfo.Hash = FMD5Hash::HashFile(*FullFilePath);
+			Statement_UpdateFileInfo.BindAndExecuteSingle(PackageName.ToLower(), OutFileInfo.LastModified, LexToString(OutFileInfo.Hash));
+			return true;
+		}
+		else
+		{
+			OutFileInfo = FileInfo.ToAssetFileInfo();
+			OutFileInfo.LastModified = CurrentLastModified;
+			OutFileInfo.Hash = FMD5Hash::HashFile(*FullFilePath);
+			Statement_AddFileInfo.BindAndExecuteSingle(PackageName.ToLower(), OutFileInfo.LastModified, LexToString(OutFileInfo.Hash));
+			return true;
+		}
+	}
+};
+
+class FFileInfoDatabaseScopedTransaction
+{
+public:
+	explicit FFileInfoDatabaseScopedTransaction(FFileInfoDatabaseStatements& InStatements)
+		: Statements(InStatements)
+		, bHasTransaction(Statements.BeginTransaction()) // This will fail if a transaction is already open
+	{
+	}
+
+	~FFileInfoDatabaseScopedTransaction()
+	{
+		Commit();
+	}
+
+	bool CommitOrRollback(const bool bShouldCommit)
+	{
+		if (bShouldCommit)
+		{
+			Commit();
+			return true;
+		}
+
+		Rollback();
+		return false;
+	}
+
+	void Commit()
+	{
+		if (bHasTransaction)
+		{
+			verify(Statements.CommitTransaction());
+			bHasTransaction = false;
+		}
+	}
+
+	void Rollback()
+	{
+		if (bHasTransaction)
+		{
+			verify(Statements.RollbackTransaction());
+			bHasTransaction = false;
+		}
+	}
+
+private:
+	FFileInfoDatabaseStatements& Statements;
+	bool bHasTransaction;
+};
+
+FFileInfoDatabase::FFileInfoDatabase()
+	: Database(MakeUnique<FSQLiteDatabase>())
+	, DatabaseFileName(TEXT("FileInfo.db"))
+{
+}
+
+FFileInfoDatabase::~FFileInfoDatabase()
+{
+	Close();
+}
+
+bool FFileInfoDatabase::IsValid() const
+{
+	return Database->IsValid();
+}
+
+bool FFileInfoDatabase::Open(const FString& InSessionPath)
+{
+	return Open(InSessionPath, ESQLiteDatabaseOpenMode::ReadWriteCreate);
+}
+
+bool FFileInfoDatabase::Open(const FString& InSessionPath, const ESQLiteDatabaseOpenMode InOpenMode)
+{
+	if (Database->IsValid())
+	{
+		return false;
+	}
+
+	if (!Database->Open(*(InSessionPath / DatabaseFileName), ESQLiteDatabaseOpenMode::ReadWriteCreate))
+	{
+		UE_LOG(LogFileInfo, Error, TEXT("Failed to open database for '%s': %s"), *InSessionPath, *GetLastError());
+		return false;
+	}
+
+	SessionPath = InSessionPath;
+
+	// Set the database to use exclusive WAL mode for performance (exclusive works even on platforms without a mmap implementation)
+	// Set the database "NORMAL" fsync mode to only perform a fsync when check-pointing the WAL to the main database file (fewer fsync calls are better for performance, with a very slight loss of WAL durability if the power fails)
+	Database->Execute(TEXT("PRAGMA cache_size=1000;"));
+	Database->Execute(TEXT("PRAGMA page_size=65535;"));
+	Database->Execute(TEXT("PRAGMA locking_mode=EXCLUSIVE;"));
+
+	//Database->Execute(TEXT("PRAGMA journal_mode=WAL;"));
+	//Database->Execute(TEXT("PRAGMA synchronous=NORMAL;"));
+
+	Database->Execute(TEXT("PRAGMA journal_mode=NORMAL;"));
+	Database->Execute(TEXT("PRAGMA synchronous=OFF;"));
+
+	int32 LoadedDatabaseVersion = 0;
+	Database->GetUserVersion(LoadedDatabaseVersion);
+	if (LoadedDatabaseVersion != (int32)EFileInfoDatabaseVersion::Empty)
+	{
+		if (LoadedDatabaseVersion > (int32)EFileInfoDatabaseVersion::LatestVersion)
+		{
+			Close();
+			UE_LOG(LogFileInfo, Error, TEXT("Failed to open database for '%s': Database is too new (version %d, expected = %d)"), *InSessionPath, LoadedDatabaseVersion, (int32)EFileInfoDatabaseVersion::LatestVersion);
+			return false;
+		}
+		else if (LoadedDatabaseVersion < (int32)EFileInfoDatabaseVersion::LatestVersion)
+		{
+			Close(true);
+			UE_LOG(LogFileInfo, Log, TEXT("Opened database '%s': Database is too old (version %d, expected = %d), creating new database"), *InSessionPath, LoadedDatabaseVersion, (int32)EFileInfoDatabaseVersion::LatestVersion);
+			return Open(InSessionPath, InOpenMode);
+		}
+	}
+
+	// Create our required tables
+	//========================================================================
+	if (!ensure(Database->Execute(TEXT("CREATE TABLE IF NOT EXISTS table_files(fileid INTEGER PRIMARY KEY, file_path TEXT UNIQUE, file_last_modified INTEGER NOT NULL, file_hash);"))))
+	{
+		LogLastError();
+		Close();
+		return false;
+	}
+
+	if (!ensure(Database->Execute(
+		TEXT("CREATE UNIQUE INDEX IF NOT EXISTS file_path_index ON table_files(file_path);")
+	)))
+	{
+		LogLastError();
+		Close();
+		return false;
+	}
+
+	// The database will have the latest schema at this point, so update the user-version
+	if (!Database->SetUserVersion((int32)EFileInfoDatabaseVersion::LatestVersion))
+	{
+		Close();
+		return false;
+	}
+
+	// Create our required prepared statements
+	Statements = MakeUnique<FFileInfoDatabaseStatements>(*Database);
+	if (!ensure(Statements->CreatePreparedStatements()))
+	{
+		Close();
+		return false;
+	}
+
+	return true;
+}
+
+bool FFileInfoDatabase::Close(const bool InDeleteDatabase)
+{
+	if (!Database->IsValid())
+	{
+		return false;
+	}
+
+	// Need to destroy prepared statements before the database can be closed
+	Statements.Reset();
+
+	if (!Database->Close())
+	{
+		UE_LOG(LogFileInfo, Error, TEXT("Failed to close database for '%s': %s"), *SessionPath, *GetLastError());
+		return false;
+	}
+
+	if (InDeleteDatabase)
+	{
+		IFileManager::Get().Delete(*(SessionPath / DatabaseFileName), false);
+	}
+
+	SessionPath.Reset();
+
+	return true;
+}
+
+FString FFileInfoDatabase::GetFilename() const
+{
+	return Database->GetFilename();
+}
+
+FString FFileInfoDatabase::GetLastError() const
+{
+	return Database->GetLastError();
+}
+
+void FFileInfoDatabase::LogLastError() const
+{
+	UE_LOG(LogFileInfo, Error, TEXT("Database Error: %s"), *SessionPath, *GetLastError());
+}
+
+bool FFileInfoDatabase::AddOrUpdateFileInfo(const FAssetData& InAssetData, FAssetFileInfo& OutFileInfo)
+{
+	if (ensure(Statements))
+	{
+		return Statements->AddOrUpdateFileInfo(InAssetData, OutFileInfo);
+	}
+
+	return false;
+}
+
+void FFileInfoDatabase::AddOrUpdateFileInfos(const TArray<FAssetData>& InAssets)
+{
+	for (const FAssetData& InAsset : InAssets)
+	{
+		// If it's a redirector act like it has been removed from the system,
+		// we don't want old duplicate entries for it.
+		if (InAsset.IsRedirector())
+		{
+			continue;
+		}
+
+		// Freshen hash cache
+		FAssetFileInfo FileInfo;
+		AddOrUpdateFileInfo(InAsset, FileInfo);
+	}
+}
+
+TMap<FName, FAssetFileInfo> FFileInfoDatabase::GetAllFileInfos()
+{
+	TMap<FName, FAssetFileInfo> FileInfos;
+	Statements->GetAllFileInfos([&FileInfos](FAssetFileInfo&& InResult)
+	{
+		FileInfos.Add(InResult.PackageName, InResult);
+
+		return ESQLitePreparedStatementExecuteRowResult::Continue;
+	});
+
+	return FileInfos;
+}
+
+PRAGMA_ENABLE_OPTIMIZATION
diff --git a/Engine/Plugins/Editor/AssetSearch/Source/Private/FileInfoDatabase.h b/Engine/Plugins/Editor/AssetSearch/Source/Private/FileInfoDatabase.h
new file mode 100644
--- /dev/null
+++ b/Engine/Plugins/Editor/AssetSearch/Source/Private/FileInfoDatabase.h
@@ -0,0 +1,99 @@
+// Copyright Epic Games, Inc. All Rights Reserved.
+
+#pragma once
+
+#include "CoreMinimal.h"
+#include "IAssetSearchModule.h"
+
+class FSQLiteDatabase;
+enum class ESQLiteDatabaseOpenMode : uint8;
+class FFileInfoDatabaseStatements;
+struct FAssetData;
+
+struct FAssetFileInfo
+{
+	FName PackageName;
+	FDateTime LastModified;
+	FMD5Hash Hash;
+};
+
+class FFileInfoDatabase
+{
+public:
+	FFileInfoDatabase();
+	~FFileInfoDatabase();
+
+	FFileInfoDatabase(const FFileInfoDatabase&) = delete;
+	FFileInfoDatabase& operator=(const FFileInfoDatabase&) = delete;
+
+	FFileInfoDatabase(FFileInfoDatabase&&) = default;
+	FFileInfoDatabase& operator=(FFileInfoDatabase&&) = default;
+
+	/**
+	 * Is this a valid database? (ie, has been successfully opened).
+	 */
+	bool IsValid() const;
+
+	/**
+	 * Open (or create) a database file.
+	 *
+	 * @param InSessionPath				The root path to store all the data for this session under.
+	 *
+	 * @return True if the database was opened, false otherwise.
+	 */
+	bool Open(const FString& InSessionPath);
+
+	/**
+	 * Open (or create) a database file.
+	 *
+	 * @param InSessionPath				The root path to store all the data for this session under.
+	 * @param InOpenMode				How should the database be opened?
+	 *
+	 * @return True if the database was opened, false otherwise.
+	 */
+	bool Open(const FString& InSessionPath, const ESQLiteDatabaseOpenMode InOpenMode);
+
+	/**
+	 * Close an open database file.
+	 *
+	 * @param InDeleteDatabase			True if the session database and its associated data should be deleted after closing the database.
+	 */
+	bool Close(const bool InDeleteDatabase = false);
+
+	/**
+	 * Get the filename of the currently open database, or an empty string.
+	 * @note The returned filename will be an absolute pathname.
+	 */
+	FString GetFilename() const;
+
+	/**
+	 * Get the last error reported by this database.
+	 */
+	FString GetLastError() const;
+
+	/**
+	 * Update the file hash table.
+	 */
+	bool AddOrUpdateFileInfo(const FAssetData& InAssetData, FAssetFileInfo& OutFileInfo);
+	void AddOrUpdateFileInfos(const TArray<FAssetData>& InAssets);
+
+	void UpdateFileHash(const FAssetData& InAssetData, FAssetFileInfo& OutFileInfo);
+
+	TMap<FName, FAssetFileInfo> GetAllFileInfos();
+
+private:
+	void LogLastError() const;
+
+private:
+	/** Internal SQLite database */
+	TUniquePtr<FSQLiteDatabase> Database;
+
+	/** The database filename name we use */
+	FString DatabaseFileName;
+
+	/** Root path to store all session data under */
+	FString SessionPath;
+
+	/** Prepared statements for the currently open database */
+	TUniquePtr<FFileInfoDatabaseStatements> Statements;
+};
\ No newline at end of file
diff --git a/Engine/Plugins/Editor/AssetSearch/Source/Private/Indexers/DataAssetIndexer.cpp b/Engine/Plugins/Editor/AssetSearch/Source/Private/Indexers/DataAssetIndexer.cpp
--- a/Engine/Plugins/Editor/AssetSearch/Source/Private/Indexers/DataAssetIndexer.cpp
+++ b/Engine/Plugins/Editor/AssetSearch/Source/Private/Indexers/DataAssetIndexer.cpp
@@ -6,8 +6,8 @@
 
 enum class EDataAssetIndexerVersion
 {
-	Empty = 0,
-	Initial = 1,
+	Empty,
+	Initial,
 
 	// -----<new versions can be added above this line>-------------------------------------------------
 	VersionPlusOne,
diff --git a/Engine/Plugins/Editor/AssetSearch/Source/Private/Indexers/DataTableIndexer.cpp b/Engine/Plugins/Editor/AssetSearch/Source/Private/Indexers/DataTableIndexer.cpp
--- a/Engine/Plugins/Editor/AssetSearch/Source/Private/Indexers/DataTableIndexer.cpp
+++ b/Engine/Plugins/Editor/AssetSearch/Source/Private/Indexers/DataTableIndexer.cpp
@@ -6,8 +6,8 @@
 
 enum class EDataTableIndexerVersion
 {
-	Empty = 0,
-	Initial = 1,
+	Empty,
+	Initial,
 
 	// -----<new versions can be added above this line>-------------------------------------------------
 	VersionPlusOne,
diff --git a/Engine/Plugins/Editor/AssetSearch/Source/Private/Indexers/DialogueWaveIndexer.cpp b/Engine/Plugins/Editor/AssetSearch/Source/Private/Indexers/DialogueWaveIndexer.cpp
--- a/Engine/Plugins/Editor/AssetSearch/Source/Private/Indexers/DialogueWaveIndexer.cpp
+++ b/Engine/Plugins/Editor/AssetSearch/Source/Private/Indexers/DialogueWaveIndexer.cpp
@@ -7,8 +7,8 @@
 
 enum class EDialogueWaveIndexerVersion
 {
-	Empty = 0,
-	Initial = 1,
+	Empty,
+	Initial,
 
 	// -----<new versions can be added above this line>-------------------------------------------------
 	VersionPlusOne,
diff --git a/Engine/Plugins/Editor/AssetSearch/Source/Private/Indexers/LevelIndexer.cpp b/Engine/Plugins/Editor/AssetSearch/Source/Private/Indexers/LevelIndexer.cpp
--- a/Engine/Plugins/Editor/AssetSearch/Source/Private/Indexers/LevelIndexer.cpp
+++ b/Engine/Plugins/Editor/AssetSearch/Source/Private/Indexers/LevelIndexer.cpp
@@ -8,8 +8,8 @@
 
 enum class ELevelIndexerVersion
 {
-	Empty = 0,
-	Initial = 1,
+	Empty,
+	Initial,
 
 	// -----<new versions can be added above this line>-------------------------------------------------
 	VersionPlusOne,
diff --git a/Engine/Plugins/Editor/AssetSearch/Source/Private/Indexers/SoundCueIndexer.cpp b/Engine/Plugins/Editor/AssetSearch/Source/Private/Indexers/SoundCueIndexer.cpp
--- a/Engine/Plugins/Editor/AssetSearch/Source/Private/Indexers/SoundCueIndexer.cpp
+++ b/Engine/Plugins/Editor/AssetSearch/Source/Private/Indexers/SoundCueIndexer.cpp
@@ -8,8 +8,8 @@
 
 enum class ESoundCueIndexerVersion
 {
-	Empty = 0,
-	Initial = 1,
+	Empty,
+	Initial,
 
 	// -----<new versions can be added above this line>-------------------------------------------------
 	VersionPlusOne,
diff --git a/Engine/Plugins/Editor/AssetSearch/Source/Private/Indexers/WidgetBlueprintIndexer.cpp b/Engine/Plugins/Editor/AssetSearch/Source/Private/Indexers/WidgetBlueprintIndexer.cpp
--- a/Engine/Plugins/Editor/AssetSearch/Source/Private/Indexers/WidgetBlueprintIndexer.cpp
+++ b/Engine/Plugins/Editor/AssetSearch/Source/Private/Indexers/WidgetBlueprintIndexer.cpp
@@ -10,9 +10,9 @@
 
 enum class EWidgetBlueprintIndexerVersion
 {
-	Empty = 0,
-	Initial = 1,
-	FixLabels = 2,
+	Empty,
+	Initial,
+	FixLabels,
 
 	// -----<new versions can be added above this line>-------------------------------------------------
 	VersionPlusOne,
diff --git a/Engine/Plugins/Editor/AssetSearch/Source/Private/SearchSerializer.cpp b/Engine/Plugins/Editor/AssetSearch/Source/Private/SearchSerializer.cpp
--- a/Engine/Plugins/Editor/AssetSearch/Source/Private/SearchSerializer.cpp
+++ b/Engine/Plugins/Editor/AssetSearch/Source/Private/SearchSerializer.cpp
@@ -13,8 +13,8 @@
 
 enum class ESearchIndexVersion
 {
-	Empty = 0,
-	Initial = 1,
+	Empty,
+	Initial,
 
 	// -----<new versions can be added above this line>-------------------------------------------------
 	VersionPlusOne,
diff --git a/Engine/Plugins/Editor/AssetSearch/Source/Private/Widgets/SSearchBrowser.cpp b/Engine/Plugins/Editor/AssetSearch/Source/Private/Widgets/SSearchBrowser.cpp
--- a/Engine/Plugins/Editor/AssetSearch/Source/Private/Widgets/SSearchBrowser.cpp
+++ b/Engine/Plugins/Editor/AssetSearch/Source/Private/Widgets/SSearchBrowser.cpp
@@ -32,6 +32,8 @@
 #include "Widgets/Views/STreeView.h"
 #include "Widgets/Layout/SSeparator.h"
 #include "Widgets/Input/SHyperlink.h"
+#include "Widgets/Images/SThrobber.h"
+#include "Widgets/Input/SSearchBox.h"
 
 #define LOCTEXT_NAMESPACE "SObjectBrowser"
 
@@ -81,11 +83,12 @@ void SSearchBrowser::Construct( const FArguments& InArgs )
 				+ SHorizontalBox::Slot()
 				.FillWidth(1.0f)
 				[
-					SNew(SEditableTextBox)
-					.ClearKeyboardFocusOnCommit(false)
+					SNew(SSearchBox)
 					.HintText(LOCTEXT("SearchHint", "Search"))
 					.OnTextCommitted(this, &SSearchBrowser::OnSearchTextCommited)
 					.OnTextChanged(this, &SSearchBrowser::OnSearchTextChanged)
+					.IsSearching(this, &SSearchBrowser::IsSearching)
+					.DelayChangeNotificationsWhileTyping(true)
 				]
 			]
 
@@ -114,31 +117,6 @@ void SSearchBrowser::Construct( const FArguments& InArgs )
 						+ SHeaderRow::Column(SSearchTreeRow::NAME_ColumnType)
 						.ManualWidth(300)
 						.DefaultLabel(LOCTEXT("ColumnType", "Type"))
-						//.SortMode(this, &SSearchBrowser::GetColumnSortMode, SObjectBrowserTableRow::CategoryClass)
-						//.OnSort(this, &SSearchBrowser::OnColumnSortModeChanged)
-
-					//	+ SHeaderRow::Column(SObjectBrowserTableRow::CategoryWorld)
-					//	.FixedWidth(22)
-					//	.DefaultLabel(LOCTEXT("CategoryWorld", ""))
-					//	.SortMode(this, &SSearchBrowser::GetColumnSortMode, SObjectBrowserTableRow::CategoryWorld)
-					//	.OnSort(this, &SSearchBrowser::OnColumnSortModeChanged)
-
-					//	+ SHeaderRow::Column(SObjectBrowserTableRow::CategoryClass)
-					//	.ManualWidth(300)
-					//	.DefaultLabel(LOCTEXT("CategoryAsset", "Asset"))
-					//	.SortMode(this, &SSearchBrowser::GetColumnSortMode, SObjectBrowserTableRow::CategoryClass)
-					//	.OnSort(this, &SSearchBrowser::OnColumnSortModeChanged)
-
-					//	+ SHeaderRow::Column(SObjectBrowserTableRow::CategoryProperty)
-					//	.FillWidth(1.0f)
-					//	.DefaultLabel(LOCTEXT("CategoryProperty", "Text"))
-					//	.SortMode(this, &SSearchBrowser::GetColumnSortMode, SObjectBrowserTableRow::CategoryProperty)
-					//	.OnSort(this, &SSearchBrowser::OnColumnSortModeChanged)
-
-					//	//+ SHeaderRow::Column(SObjectBrowserTableRow::CategoryPropertyValue)
-					//	//.DefaultLabel(LOCTEXT("CategoryPropertyValue", "Value"))
-					//	//.SortMode(this, &SSearchBrowser::GetColumnSortMode, SObjectBrowserTableRow::CategoryPropertyValue)
-					//	//.OnSort(this, &SSearchBrowser::OnColumnSortModeChanged)
 					)
 				]
 			]
@@ -238,8 +216,11 @@ void SSearchBrowser::RefreshList()
 		FSearchQuery Query;
 		Query.Query = FilterText.ToString();
 
+		SearchesActive++;
+
 		IAssetSearchModule& SearchModule = IAssetSearchModule::Get();
 		SearchModule.Search(Query, [this](TArray<FSearchRecord>&& InResults) {
+			check(IsInGameThread());
 
 			SearchResults.Reset();
 			SearchResultHierarchy.Reset();
@@ -261,6 +242,7 @@ void SSearchBrowser::RefreshList()
 			});
 
 			SearchTreeView->RequestListRefresh();
+			SearchesActive--;
 		});
 	}
 }
@@ -286,7 +268,13 @@ void SSearchBrowser::OnSearchTextCommited(const FText& InText, ETextCommit::Type
 
 void SSearchBrowser::OnSearchTextChanged(const FText& InText)
 {
-	if (InText.ToString().Len() > 3)
+	const int32 Length = InText.ToString().Len();
+
+	if (Length > 4)
+	{
+		TryRefreshingSearch(InText);
+	}
+	else if (Length == 0)
 	{
 		TryRefreshingSearch(InText);
 	}
@@ -319,4 +307,9 @@ void SSearchBrowser::HandleListSelectionChanged(TSharedPtr<FSearchNode> InNode,
 {
 }
 
+bool SSearchBrowser::IsSearching() const
+{
+	return SearchesActive > 0;
+}
+
 #undef LOCTEXT_NAMESPACE
diff --git a/Engine/Plugins/Editor/AssetSearch/Source/Private/Widgets/SSearchBrowser.h b/Engine/Plugins/Editor/AssetSearch/Source/Private/Widgets/SSearchBrowser.h
--- a/Engine/Plugins/Editor/AssetSearch/Source/Private/Widgets/SSearchBrowser.h
+++ b/Engine/Plugins/Editor/AssetSearch/Source/Private/Widgets/SSearchBrowser.h
@@ -62,19 +62,23 @@ private:
 
 	void HandleListSelectionChanged(TSharedPtr<FSearchNode> TransactionInfo, ESelectInfo::Type SelectInfo);
 
+	bool IsSearching() const;
+
 private:
 
 	FText FilterText;
 
 	// Filters
 	FString FilterString;
+
+	int32 SearchesActive = 0;
 	
 	TMap<FString, TSharedPtr<FAssetNode>> SearchResultHierarchy;
 	TArray< TSharedPtr<FSearchNode> > SearchResults;
 
 	TSharedPtr< STreeView< TSharedPtr<FSearchNode> > > SearchTreeView;
 
-	IAssetRegistry* AssetRegistry;
+	IAssetRegistry* AssetRegistry = nullptr;
 
 	TSharedPtr<FAssetThumbnailPool> ThumbnailPool;
 
