Universal Search - Adding options for users.  US is now opt-in.  Opening the search window or enabling it in your project settings will enable it.

diff --git a/Engine/Plugins/Editor/AssetSearch/Source/Private/AssetSearchManager.cpp b/Engine/Plugins/Editor/AssetSearch/Source/Private/AssetSearchManager.cpp
--- a/Engine/Plugins/Editor/AssetSearch/Source/Private/AssetSearchManager.cpp
+++ b/Engine/Plugins/Editor/AssetSearch/Source/Private/AssetSearchManager.cpp
@@ -25,18 +25,28 @@
 #include "Indexers/DialogueWaveIndexer.h"
 #include "Sound/DialogueWave.h"
 #include "Indexers/LevelIndexer.h"
-#include "Settings/AssetSearchDeveloperSettings.h"
+#include "Settings/SearchProjectSettings.h"
+#include "Settings/SearchUserSettings.h"
 #include "Indexers/SoundCueIndexer.h"
 #include "Sound/SoundCue.h"
 #include "Misc/ScopedSlowTask.h"
 #include "Misc/StringBuilder.h"
 #include "Engine/World.h"
 #include "Editor.h"
+#include "PackageTools.h"
+#include "UObject/ObjectMacros.h"
 
 PRAGMA_DISABLE_OPTIMIZATION
 
 #define LOCTEXT_NAMESPACE "FAssetSearchManager"
 
+static bool bDisableUniversalSearch = false;
+FAutoConsoleVariableRef CVarDisableUniversalSearch(
+	TEXT("Search.Disable"),
+	bDisableUniversalSearch,
+	TEXT("Disable Universal Search (only works on startup)")
+);
+
 static bool bIndexUnindexAssetsOnLoad = false;
 FAutoConsoleVariableRef CVarIndexUnindexAssetsOnLoad(
 	TEXT("Search.IndexUnindexAssetsOnLoad"),
@@ -65,6 +75,97 @@ FAutoConsoleVariableRef CVarGameThread_AssetScanLimit(
 	TEXT("")
 );
 
+class FUnloadPackageScope
+{
+public:
+	FUnloadPackageScope()
+	{
+		FCoreUObjectDelegates::OnAssetLoaded.AddRaw(this, &FUnloadPackageScope::OnAssetLoaded);
+	}
+
+	~FUnloadPackageScope()
+	{
+		FCoreUObjectDelegates::OnAssetLoaded.RemoveAll(this);
+		TryUnload(true);
+	}
+
+	int32 TryUnload(bool bResetTrackedObjects)
+	{
+		TArray<TWeakObjectPtr<UObject>> PackageObjectPtrs;
+
+		for (const FObjectKey& LoadedObjectKey : ObjectsLoaded)
+		{
+			if (UObject* LoadedObject = LoadedObjectKey.ResolveObjectPtr())
+			{
+				UPackage* Package = LoadedObject->GetOutermost();
+
+				TArray<UObject*> PackageObjects;
+				GetObjectsWithOuter(Package, PackageObjects, false);
+
+				for (UObject* PackageObject : PackageObjects)
+				{
+					PackageObject->ClearFlags(RF_Standalone);
+					PackageObjectPtrs.Add(PackageObject);
+				}
+			}
+		}
+
+		CollectGarbage(GARBAGE_COLLECTION_KEEPFLAGS);
+		
+		int32 NumRemoved = 0;
+		for (int32 LoadedAssetIndex = 0; LoadedAssetIndex < PackageObjectPtrs.Num(); LoadedAssetIndex++)
+		{
+			TWeakObjectPtr<UObject> PackageObjectPtr = PackageObjectPtrs[LoadedAssetIndex];
+
+			if (UObject* LoadedObject = PackageObjectPtr.Get())
+			{
+				//FReferencerInformationList ReferencesIncludingUndo;
+				//bool bReferencedInMemoryOrUndoStack = IsReferenced(LoadedObject, GARBAGE_COLLECTION_KEEPFLAGS, EInternalObjectFlags::GarbageCollectionKeepFlags, true, &ReferencesIncludingUndo);
+
+				LoadedObject->SetFlags(RF_Standalone);
+			}
+			else
+			{
+				NumRemoved++;
+			}
+		}
+
+		if (bResetTrackedObjects)
+		{
+			ObjectsLoaded.Reset();
+		}
+		else
+		{
+			for (int32 ObjectIndex = 0; ObjectIndex < ObjectsLoaded.Num(); ObjectIndex++)
+			{
+				if (!ObjectsLoaded[ObjectIndex].ResolveObjectPtr())
+				{
+					ObjectsLoaded.RemoveAt(ObjectIndex);
+					ObjectIndex--;
+				}
+			}
+		}
+
+		return NumRemoved;
+	}
+
+	int32 GetObjectsLoaded() const
+	{
+		return ObjectsLoaded.Num();
+	}
+
+private:
+	void OnAssetLoaded(UObject* InObject)
+	{
+		ObjectsLoaded.Add(FObjectKey(InObject));
+	}
+
+private:
+	TArray<FObjectKey> ObjectsLoaded;
+	TArray<UClass*> ClassFilters;
+};
+
+
 FAssetSearchManager::FAssetSearchManager()
 {
 	PendingDatabaseUpdates = 0;
@@ -93,11 +194,12 @@ FAssetSearchManager::~FAssetSearchManager()
 
 void FAssetSearchManager::Start()
 {
-	// Don't start the search manager on unattended builds.
-	//if (FApp::IsUnattended())
-	//{
-	//	return;
-	//}
+	TickerHandle = FTicker::GetCoreTicker().AddTicker(FTickerDelegate::CreateRaw(this, &FAssetSearchManager::Tick_GameThread), 0);
+}
+
+void FAssetSearchManager::InternalStart()
+{
+	bStarted = true;
 
 	RegisterAssetIndexer(UDataAsset::StaticClass(), MakeUnique<FDataAssetIndexer>());
 	RegisterAssetIndexer(UDataTable::StaticClass(), MakeUnique<FDataTableIndexer>());
@@ -123,8 +225,6 @@ void FAssetSearchManager::Start()
 		OnAssetAdded(Data);
 	}
 
-	TickerHandle = FTicker::GetCoreTicker().AddTicker(FTickerDelegate::CreateRaw(this, &FAssetSearchManager::Tick_GameThread), 0);
-
 	RunThread = true;
 	DatabaseThread = FRunnableThread::Create(this, TEXT("UniversalSearch"), 0, TPri_BelowNormal);
 }
@@ -191,8 +291,18 @@ void FAssetSearchManager::OnAssetAdded(const FAssetData& InAssetData)
 	}
 
 	// 
-	const UAssetSearchDeveloperSettings* Settings = GetDefault<UAssetSearchDeveloperSettings>();
-	for (const FDirectoryPath& IgnoredPath : Settings->IgnoredPaths)
+	const USearchProjectSettings* ProjectSettings = GetDefault<USearchProjectSettings>();
+	for (const FDirectoryPath& IgnoredPath : ProjectSettings->IgnoredPaths)
+	{
+		if (PackageName.StartsWith(IgnoredPath.Path))
+		{
+			return;
+		}
+	}
+
+	// 
+	const USearchUserSettings* UserSettings = GetDefault<USearchUserSettings>();
+	for (const FDirectoryPath& IgnoredPath : UserSettings->IgnoredPaths)
 	{
 		if (PackageName.StartsWith(IgnoredPath.Path))
 		{
@@ -478,6 +588,16 @@ bool FAssetSearchManager::Tick_GameThread(float DeltaTime)
 {
 	check(IsInGameThread());
 
+	if (bDisableUniversalSearch || !GetDefault<USearchUserSettings>()->bEnableSearch)
+	{
+		return true;
+	}
+
+	if (!bStarted)
+	{
+		InternalStart();
+	}
+
 	TryConnectToDatabase();
 
 	//if (0)
@@ -494,6 +614,8 @@ bool FAssetSearchManager::Tick_GameThread(float DeltaTime)
 
 	ProcessGameThreadTasks();
 
+	const USearchUserSettings* UserSettings = GetDefault<USearchUserSettings>();
+
 	int32 ScanLimit = GameThread_AssetScanLimit;
 	while (ProcessAssetQueue.Num() > 0 && ScanLimit > 0 && PendingDownloads < PendingDownloadsMax)
 	{
@@ -596,6 +718,8 @@ void FAssetSearchManager::ForceIndexOnAssetsMissingIndex()
 	IndexingTask.MakeDialog(true);
 
 	int32 RemovedCount = 0;
+
+	FUnloadPackageScope UnloadScope;
 	for (const FAssetDDCRequest& Request : FailedDDCRequests)
 	{
 		if (IndexingTask.ShouldCancel())
@@ -609,6 +733,11 @@ void FAssetSearchManager::ForceIndexOnAssetsMissingIndex()
 			StoreIndexForAsset(AssetToIndex);
 		}
 
+		if (UnloadScope.GetObjectsLoaded() > 2000)
+		{
+			UnloadScope.TryUnload(true);
+		}
+
 		RemovedCount++;
 	}
 
diff --git a/Engine/Plugins/Editor/AssetSearch/Source/Private/AssetSearchManager.h b/Engine/Plugins/Editor/AssetSearch/Source/Private/AssetSearchManager.h
--- a/Engine/Plugins/Editor/AssetSearch/Source/Private/AssetSearchManager.h
+++ b/Engine/Plugins/Editor/AssetSearch/Source/Private/AssetSearchManager.h
@@ -36,6 +36,8 @@ private:
 	void Tick_DatabaseOperationThread();
 
 private:
+	void InternalStart();
+
 	void OnAssetAdded(const FAssetData& InAssetData);
 	void OnAssetRemoved(const FAssetData& InAssetData);
 	void OnAssetScanFinished();
@@ -57,6 +59,9 @@ private:
 	void AsyncMainThreadTask(TFunction<void()> Task);
 	void ProcessGameThreadTasks();
 
+private:
+	bool bStarted = false;
+
 private:
 	FFileInfoDatabase FileInfoDatabase;
 	FCriticalSection FileInfoDatabaseCS;
diff --git a/Engine/Plugins/Editor/AssetSearch/Source/Private/Indexers/BlueprintIndexer.cpp b/Engine/Plugins/Editor/AssetSearch/Source/Private/Indexers/BlueprintIndexer.cpp
--- a/Engine/Plugins/Editor/AssetSearch/Source/Private/Indexers/BlueprintIndexer.cpp
+++ b/Engine/Plugins/Editor/AssetSearch/Source/Private/Indexers/BlueprintIndexer.cpp
@@ -9,6 +9,10 @@
 #include "Utility/IndexerUtilities.h"
 #include "K2Node_BaseMCDelegate.h"
 #include "Internationalization/Text.h"
+#include "K2Node_Knot.h"
+#include "EdGraphNode_Comment.h"
+#include "Engine/SimpleConstructionScript.h"
+#include "Engine/SCS_Node.h"
 
 #define LOCTEXT_NAMESPACE "FBlueprintIndexer"
 
@@ -20,6 +24,7 @@ enum class EBlueprintIndexerVersion
 	Initial,
 	FixingPinsToSaveValues,
 	IndexingPublicEditableFieldsOnNodes,
+	DontIndexPinsUnlessItsInputWithNoConnections,
 
 	// -----<new versions can be added above this line>-------------------------------------------------
 	VersionPlusOne,
@@ -36,7 +41,20 @@ void FBlueprintIndexer::IndexAsset(const UObject* InAssetObject, FSearchSerializ
 	const UBlueprint* BP = Cast<UBlueprint>(InAssetObject);
 	check(BP);
 
-	if (UClass* GeneratedClass = BP->GeneratedClass)
+	Serializer.BeginIndexingObject(BP, TEXT("$self"));
+	FIndexerUtilities::IterateIndexableProperties(BP, [&Serializer](const FProperty* Property, const FString& Value) {
+		Serializer.IndexProperty(Property, Value);
+	});
+	Serializer.EndIndexingObject();
+
+	IndexClassDefaultObject(BP, Serializer);
+	IndexComponents(BP, Serializer);
+	IndexGraphs(BP, Serializer);
+}
+
+void FBlueprintIndexer::IndexClassDefaultObject(const UBlueprint* InBlueprint, FSearchSerializer& Serializer) const
+{
+	if (UClass* GeneratedClass = InBlueprint->GeneratedClass)
 	{
 		if (UObject* CDO = GeneratedClass->GetDefaultObject())
 		{
@@ -47,14 +65,48 @@ void FBlueprintIndexer::IndexAsset(const UObject* InAssetObject, FSearchSerializ
 			Serializer.EndIndexingObject();
 		}
 	}
+}
 
+void FBlueprintIndexer::IndexComponents(const UBlueprint* InBlueprint, FSearchSerializer& Serializer) const
+{
+	if (InBlueprint->SimpleConstructionScript)
+	{
+		const TArray<USCS_Node*>& BPNodes = InBlueprint->SimpleConstructionScript->GetAllNodes();
+		for (USCS_Node* BPNode : BPNodes)
+		{
+			Serializer.BeginIndexingObject(BPNode->ComponentTemplate, BPNode->GetVariableName().ToString());
+			FIndexerUtilities::IterateIndexableProperties(BPNode->ComponentTemplate, [&Serializer](const FProperty* Property, const FString& Value) {
+				Serializer.IndexProperty(Property, Value);
+			});
+			Serializer.EndIndexingObject();
+		}
+	}
+}
+
+void FBlueprintIndexer::IndexGraphs(const UBlueprint* InBlueprint, FSearchSerializer& Serializer) const
+{
 	TArray<UEdGraph*> AllGraphs;
-	BP->GetAllGraphs(AllGraphs);
+	InBlueprint->GetAllGraphs(AllGraphs);
 
 	for (UEdGraph* Graph : AllGraphs)
-	{	
+	{
 		for (UEdGraphNode* Node : Graph->Nodes)
 		{
+			// Ignore Knots.
+			if (Cast<UK2Node_Knot>(Node))
+			{
+				continue;
+			}
+
+			// Special rules for comment nodes
+			if (UEdGraphNode_Comment* CommentNode = Cast<UEdGraphNode_Comment>(Node))
+			{
+				Serializer.BeginIndexingObject(Node, Node->NodeComment);
+				Serializer.IndexProperty(TEXT("Comment"), Node->NodeComment);
+				Serializer.EndIndexingObject();
+				continue;
+			}
+
 			const FText NodeText = Node->GetNodeTitle(ENodeTitleType::MenuTitle);
 			Serializer.BeginIndexingObject(Node, NodeText);
 			Serializer.IndexProperty(TEXT("Name"), NodeText);
@@ -63,7 +115,7 @@ void FBlueprintIndexer::IndexAsset(const UObject* InAssetObject, FSearchSerializ
 			{
 				Serializer.IndexProperty(TEXT("Comment"), Node->NodeComment);
 			}
-			
+
 			if (UK2Node_CallFunction* FunctionNode = Cast<UK2Node_CallFunction>(Node))
 			{
 				IndexMemberReference(Serializer, FunctionNode->FunctionReference, TEXT("Function"));
@@ -75,12 +127,11 @@ void FBlueprintIndexer::IndexAsset(const UObject* InAssetObject, FSearchSerializ
 			else if (UK2Node_Variable* VariableNode = Cast<UK2Node_Variable>(Node))
 			{
 				IndexMemberReference(Serializer, VariableNode->VariableReference, TEXT("Variable"));
-				//Serializer.WriteValue(TEXT("bSelfContext"), VariableReference.IsSelfContext());
 			}
 
-			if (Node->GetAllPins().Num())
+			for (UEdGraphPin* Pin : Node->GetAllPins())
 			{
-				for (UEdGraphPin* Pin : Node->GetAllPins())
+				if (Pin->Direction == EGPD_Input && Pin->LinkedTo.Num() == 0)
 				{
 					const FText PinText = Pin->GetDisplayName();
 					if (PinText.IsEmpty())
@@ -88,7 +139,7 @@ void FBlueprintIndexer::IndexAsset(const UObject* InAssetObject, FSearchSerializ
 						continue;
 					}
 
-					const FString PinValue = Pin->DefaultValue;
+					const FText PinValue = Pin->GetDefaultAsText();
 					if (PinValue.IsEmpty())
 					{
 						continue;
diff --git a/Engine/Plugins/Editor/AssetSearch/Source/Private/Indexers/BlueprintIndexer.h b/Engine/Plugins/Editor/AssetSearch/Source/Private/Indexers/BlueprintIndexer.h
--- a/Engine/Plugins/Editor/AssetSearch/Source/Private/Indexers/BlueprintIndexer.h
+++ b/Engine/Plugins/Editor/AssetSearch/Source/Private/Indexers/BlueprintIndexer.h
@@ -15,5 +15,8 @@ public:
 	virtual void IndexAsset(const UObject* InAssetObject, FSearchSerializer& Serializer) const override;
 
 private:
+	void IndexClassDefaultObject(const UBlueprint* InBlueprint, FSearchSerializer& Serializer) const;
+	void IndexComponents(const UBlueprint* InBlueprint, FSearchSerializer& Serializer) const;
+	void IndexGraphs(const UBlueprint* InBlueprint, FSearchSerializer& Serializer) const;
 	void IndexMemberReference(FSearchSerializer& Serializer, const FMemberReference& MemberReference, const FString& MemberType) const;
 };
\ No newline at end of file
diff --git a/Engine/Plugins/Editor/AssetSearch/Source/Private/Settings/AssetSearchDeveloperSettings.cpp b/Engine/Plugins/Editor/AssetSearch/Source/Private/Settings/AssetSearchDeveloperSettings.cpp
deleted file mode 100644
diff --git a/Engine/Plugins/Editor/AssetSearch/Source/Private/Settings/AssetSearchDeveloperSettings.h b/Engine/Plugins/Editor/AssetSearch/Source/Private/Settings/AssetSearchDeveloperSettings.h
deleted file mode 100644
diff --git a/Engine/Plugins/Editor/AssetSearch/Source/Private/Settings/SearchProjectSettings.cpp b/Engine/Plugins/Editor/AssetSearch/Source/Private/Settings/SearchProjectSettings.cpp
new file mode 100644
--- /dev/null
+++ b/Engine/Plugins/Editor/AssetSearch/Source/Private/Settings/SearchProjectSettings.cpp
@@ -0,0 +1,7 @@
+// Copyright Epic Games, Inc. All Rights Reserved.
+
+#include "SearchProjectSettings.h"
+
+USearchProjectSettings::USearchProjectSettings()
+{
+}
diff --git a/Engine/Plugins/Editor/AssetSearch/Source/Private/Settings/SearchProjectSettings.h b/Engine/Plugins/Editor/AssetSearch/Source/Private/Settings/SearchProjectSettings.h
new file mode 100644
--- /dev/null
+++ b/Engine/Plugins/Editor/AssetSearch/Source/Private/Settings/SearchProjectSettings.h
@@ -0,0 +1,19 @@
+// Copyright Epic Games, Inc. All Rights Reserved.
+
+#pragma once
+
+#include "Engine/DeveloperSettings.h"
+#include "Engine/EngineTypes.h"
+#include "SearchProjectSettings.generated.h"
+
+UCLASS(config = Editor, defaultconfig, meta=(DisplayName="Search"))
+class USearchProjectSettings : public UDeveloperSettings
+{
+	GENERATED_BODY()
+
+public:
+	USearchProjectSettings();
+
+	UPROPERTY(config, EditAnywhere, Category=General)
+	TArray<FDirectoryPath> IgnoredPaths;
+};
diff --git a/Engine/Plugins/Editor/AssetSearch/Source/Private/Settings/SearchUserSettings.cpp b/Engine/Plugins/Editor/AssetSearch/Source/Private/Settings/SearchUserSettings.cpp
new file mode 100644
--- /dev/null
+++ b/Engine/Plugins/Editor/AssetSearch/Source/Private/Settings/SearchUserSettings.cpp
@@ -0,0 +1,7 @@
+// Copyright Epic Games, Inc. All Rights Reserved.
+
+#include "SearchUserSettings.h"
+
+USearchUserSettings::USearchUserSettings()
+{
+}
diff --git a/Engine/Plugins/Editor/AssetSearch/Source/Private/Settings/SearchUserSettings.h b/Engine/Plugins/Editor/AssetSearch/Source/Private/Settings/SearchUserSettings.h
new file mode 100644
--- /dev/null
+++ b/Engine/Plugins/Editor/AssetSearch/Source/Private/Settings/SearchUserSettings.h
@@ -0,0 +1,26 @@
+// Copyright Epic Games, Inc. All Rights Reserved.
+
+#pragma once
+
+#include "Engine/DeveloperSettings.h"
+#include "Engine/EngineTypes.h"
+#include "SearchUserSettings.generated.h"
+
+UCLASS(config = EditorPerProjectUserSettings, meta=(DisplayName="Search"))
+class USearchUserSettings : public UDeveloperSettings
+{
+	GENERATED_BODY()
+
+public:
+	USearchUserSettings();
+
+	/** Enable this to begin using search. */
+	UPROPERTY(config, EditAnywhere, Category = General)
+	bool bEnableSearch;
+
+	UPROPERTY(config, EditAnywhere, Category=General)
+	TArray<FDirectoryPath> IgnoredPaths;
+
+	UPROPERTY(config, EditAnywhere, AdvancedDisplay, Category=General)
+	bool bShowAdvancedData;
+};
diff --git a/Engine/Plugins/Editor/AssetSearch/Source/Private/Widgets/SSearchBrowser.cpp b/Engine/Plugins/Editor/AssetSearch/Source/Private/Widgets/SSearchBrowser.cpp
--- a/Engine/Plugins/Editor/AssetSearch/Source/Private/Widgets/SSearchBrowser.cpp
+++ b/Engine/Plugins/Editor/AssetSearch/Source/Private/Widgets/SSearchBrowser.cpp
@@ -34,6 +34,7 @@
 #include "Widgets/Input/SHyperlink.h"
 #include "Widgets/Images/SThrobber.h"
 #include "Widgets/Input/SSearchBox.h"
+#include "Settings/SearchUserSettings.h"
 
 #define LOCTEXT_NAMESPACE "SObjectBrowser"
 
@@ -45,8 +46,16 @@ namespace AssetSearchConstants
 	const int32 ThumbnailPoolSize = 64;
 }
 
+//TODO Expose TSharedRef<SWidget> SAssetViewItem::CreateToolTipWidget() const via IContentBrowserSingleton.
+
 void SSearchBrowser::Construct( const FArguments& InArgs )
 {
+	if (!GetDefault<USearchUserSettings>()->bEnableSearch)
+	{
+		GetMutableDefault<USearchUserSettings>()->bEnableSearch = true;
+		GetMutableDefault<USearchUserSettings>()->SaveConfig();
+	}
+
 	SortByColumn = SSearchTreeRow::NAME_ColumnName;
 	SortMode = EColumnSortMode::Ascending;
 
@@ -140,13 +149,29 @@ void SSearchBrowser::Construct( const FArguments& InArgs )
 				.VAlign(VAlign_Center)
 				.Padding(8, 0)
 				[
-					SNew(STextBlock)
-					.Text(this, &SSearchBrowser::GetStatusText)
+					SNew(SVerticalBox)
+
+					+ SVerticalBox::Slot()
+					.AutoHeight()
+					[
+						SNew(STextBlock)
+						.Visibility_Lambda([] { return GetDefault<USearchUserSettings>()->bShowAdvancedData ? EVisibility::Visible : EVisibility::Collapsed; })
+						.Text(this, &SSearchBrowser::GetStatusText)
+					]
+
+					+ SVerticalBox::Slot()
+					.AutoHeight()
+					[
+						SNew(STextBlock)
+						.Visibility_Lambda([] { return GetDefault<USearchUserSettings>()->bShowAdvancedData ? EVisibility::Visible : EVisibility::Collapsed; })
+						.Text(this, &SSearchBrowser::GetAdvancedStatus)
+					]
 				]
 
 				// Index unindexed items
 				+ SHorizontalBox::Slot()
 				.AutoWidth()
+				.VAlign(VAlign_Top)
 				[
 					SNew(SHyperlink)
 					.Text(this, &SSearchBrowser::GetUnindexedAssetsText)
@@ -163,7 +188,23 @@ FText SSearchBrowser::GetStatusText() const
 {
 	IAssetSearchModule& SearchModule = IAssetSearchModule::Get();
 	FSearchStats SearchStats = SearchModule.GetStats();
-	return FText::Format(LOCTEXT("SearchStatusTextFmt", "Scanning {0}   Downloading {1}   Updating {2}            Total Records {3}"), SearchStats.Scanning, SearchStats.Downloading, SearchStats.PendingDatabaseUpdates, SearchStats.TotalRecords);
+	int32 UpdatingCount = SearchStats.Scanning + SearchStats.Downloading + SearchStats.PendingDatabaseUpdates;
+
+	if (UpdatingCount > 0)
+	{
+		return LOCTEXT("Updating", "Updating...  (You can search any time)");
+	}
+	else
+	{
+		return LOCTEXT("Ready", "Ready");
+	}
+}
+
+FText SSearchBrowser::GetAdvancedStatus() const
+{
+	IAssetSearchModule& SearchModule = IAssetSearchModule::Get();
+	FSearchStats SearchStats = SearchModule.GetStats();
+	return FText::Format(LOCTEXT("AdvancedSearchStatusTextFmt", "Scanning {0}   Downloading {1}   Updating {2}            Total Records {3}"), SearchStats.Scanning, SearchStats.Downloading, SearchStats.PendingDatabaseUpdates, SearchStats.TotalRecords);
 }
 
 FText SSearchBrowser::GetUnindexedAssetsText() const
@@ -270,7 +311,7 @@ void SSearchBrowser::OnSearchTextChanged(const FText& InText)
 {
 	const int32 Length = InText.ToString().Len();
 
-	if (Length > 4)
+	if (Length > 3)
 	{
 		TryRefreshingSearch(InText);
 	}
diff --git a/Engine/Plugins/Editor/AssetSearch/Source/Private/Widgets/SSearchBrowser.h b/Engine/Plugins/Editor/AssetSearch/Source/Private/Widgets/SSearchBrowser.h
--- a/Engine/Plugins/Editor/AssetSearch/Source/Private/Widgets/SSearchBrowser.h
+++ b/Engine/Plugins/Editor/AssetSearch/Source/Private/Widgets/SSearchBrowser.h
@@ -41,6 +41,7 @@ public:
 private:
 
 	FText GetStatusText() const;
+	FText GetAdvancedStatus() const;
 	FText GetUnindexedAssetsText() const;
 
 	void HandleForceIndexOfAssetsMissingIndex();
diff --git a/Engine/Source/Editor/EditorStyle/Private/SlateEditorStyle.cpp b/Engine/Source/Editor/EditorStyle/Private/SlateEditorStyle.cpp
--- a/Engine/Source/Editor/EditorStyle/Private/SlateEditorStyle.cpp
+++ b/Engine/Source/Editor/EditorStyle/Private/SlateEditorStyle.cpp
@@ -4678,11 +4678,14 @@ void FSlateEditorStyle::FStyle::SetupGraphEditorStyles()
 		}
 
 		Set("ClassIcon.K2Node_CallFunction",	new IMAGE_BRUSH("Icons/icon_Blueprint_NewFunction_16x", Icon16x16));
+		Set("ClassIcon.K2Node_FunctionEntry",	new IMAGE_BRUSH("Icons/icon_Blueprint_NewFunction_16x", Icon16x16));
 		Set("ClassIcon.K2Node_CustomEvent",		new IMAGE_BRUSH("Icons/icon_Blueprint_Event_16x", Icon16x16));
 		Set("ClassIcon.K2Node_Event",			new IMAGE_BRUSH("Icons/icon_Blueprint_Event_16x", Icon16x16));
 		Set("ClassIcon.K2Node_Variable",		new IMAGE_BRUSH("Graph/Icons/Node", Icon16x16, FLinearColor::White));
 		Set("ClassIcon.K2Node_VariableGet",		new IMAGE_BRUSH("Graph/Icons/FIB_VarGet", Icon16x16, FLinearColor::White));
 		Set("ClassIcon.K2Node_VariableSet",		new IMAGE_BRUSH("Graph/Icons/FIB_VarSet", Icon16x16, FLinearColor::White));
+		Set("ClassIcon.K2Node_DynamicCast",		new IMAGE_BRUSH("Icons/icon_Blueprint_Cast_16x", Icon16x16));
+		Set("ClassIcon.EdGraphNode_Comment",	new IMAGE_BRUSH("Icons/icon_Blueprint_Comment_16x", Icon16x16));
 
 
 		Set( "GraphEditor.Default_16x", new IMAGE_BRUSH("Icons/icon_Blueprint_Node_16x", Icon16x16));
