IWYU fixes.

diff --git a/Engine/Plugins/Editor/AssetSearch/Source/Private/AssetSearchManager.cpp b/Engine/Plugins/Editor/AssetSearch/Source/Private/AssetSearchManager.cpp
--- a/Engine/Plugins/Editor/AssetSearch/Source/Private/AssetSearchManager.cpp
+++ b/Engine/Plugins/Editor/AssetSearch/Source/Private/AssetSearchManager.cpp
@@ -35,6 +35,8 @@
 #include "Editor.h"
 #include "PackageTools.h"
 #include "UObject/ObjectMacros.h"
+#include "UObject/ObjectKey.h"
+#include "UObject/WeakObjectPtrTemplates.h"
 
 PRAGMA_DISABLE_OPTIMIZATION
 
