Universal Search - Working on improvements to scanning, providing tools for tuning how much work it does in the foreground vs background.  Now makes sure that even if users aren't enabling search, they're still seeding the DDC when they save new versions of assets so that things don't go unindexed.

diff --git a/Engine/Plugins/Editor/AssetSearch/Source/Private/AssetSearchManager.cpp b/Engine/Plugins/Editor/AssetSearch/Source/Private/AssetSearchManager.cpp
--- a/Engine/Plugins/Editor/AssetSearch/Source/Private/AssetSearchManager.cpp
+++ b/Engine/Plugins/Editor/AssetSearch/Source/Private/AssetSearchManager.cpp
@@ -42,11 +42,11 @@ PRAGMA_DISABLE_OPTIMIZATION
 
 #define LOCTEXT_NAMESPACE "FAssetSearchManager"
 
-static bool bDisableUniversalSearch = false;
+static bool bForceEnableSearch = false;
 FAutoConsoleVariableRef CVarDisableUniversalSearch(
-	TEXT("Search.Disable"),
-	bDisableUniversalSearch,
-	TEXT("Disable Universal Search (only works on startup)")
+	TEXT("Search.ForceEnable"),
+	bForceEnableSearch,
+	TEXT("Enable universal search")
 );
 
 static bool bIndexUnindexAssetsOnLoad = false;
@@ -56,27 +56,6 @@ FAutoConsoleVariableRef CVarIndexUnindexAssetsOnLoad(
 	TEXT("Index Unindex Assets On Load")
 );
 
-static int32 PendingDownloadsMax = 100;
-FAutoConsoleVariableRef CVarPendingDownloadsMax(
-	TEXT("Search.PendingDownloadsMax"),
-	PendingDownloadsMax,
-	TEXT("")
-);
-
-static int32 GameThread_DownloadProcessLimit = 30;
-FAutoConsoleVariableRef CVarGameThread_DownloadProcessLimit(
-	TEXT("Search.GameThread_DownloadProcessLimit"),
-	GameThread_DownloadProcessLimit,
-	TEXT("")
-);
-
-static int32 GameThread_AssetScanLimit = 1000;
-FAutoConsoleVariableRef CVarGameThread_AssetScanLimit(
-	TEXT("Search.GameThread_AssetScanLimit"),
-	GameThread_AssetScanLimit,
-	TEXT("")
-);
-
 class FUnloadPackageScope
 {
 public:
@@ -171,7 +150,9 @@ private:
 FAssetSearchManager::FAssetSearchManager()
 {
 	PendingDatabaseUpdates = 0;
-	PendingDownloads = 0;
+	IsAssetUpToDateCount = 0;
+	ActiveDownloads = 0;
+	DownloadQueueCount = 0;
 	TotalSearchRecords = 0;
 	LastRecordCountUpdateSeconds = 0;
 
@@ -187,6 +168,8 @@ FAssetSearchManager::~FAssetSearchManager()
 		DatabaseThread->WaitForCompletion();
 	}
 
+	StopScanningAssets();
+
 	FCoreUObjectDelegates::OnObjectSaved.RemoveAll(this);
 	FCoreUObjectDelegates::OnAssetLoaded.RemoveAll(this);
 	UObject::FAssetRegistryTag::OnGetExtraObjectTags.RemoveAll(this);
@@ -196,13 +179,6 @@ FAssetSearchManager::~FAssetSearchManager()
 
 void FAssetSearchManager::Start()
 {
-	TickerHandle = FTicker::GetCoreTicker().AddTicker(FTickerDelegate::CreateRaw(this, &FAssetSearchManager::Tick_GameThread), 0);
-}
-
-void FAssetSearchManager::InternalStart()
-{
-	bStarted = true;
-
 	RegisterAssetIndexer(UDataAsset::StaticClass(), MakeUnique<FDataAssetIndexer>());
 	RegisterAssetIndexer(UDataTable::StaticClass(), MakeUnique<FDataTableIndexer>());
 	RegisterAssetIndexer(UBlueprint::StaticClass(), MakeUnique<FBlueprintIndexer>());
@@ -214,6 +190,43 @@ void FAssetSearchManager::InternalStart()
 	FCoreUObjectDelegates::OnObjectSaved.AddRaw(this, &FAssetSearchManager::OnObjectSaved);
 	FCoreUObjectDelegates::OnAssetLoaded.AddRaw(this, &FAssetSearchManager::OnAssetLoaded);
 
+	TickerHandle = FTicker::GetCoreTicker().AddTicker(FTickerDelegate::CreateRaw(this, &FAssetSearchManager::Tick_GameThread), 0);
+
+	RunThread = true;
+	DatabaseThread = FRunnableThread::Create(this, TEXT("UniversalSearch"), 0, TPri_BelowNormal);
+}
+
+void FAssetSearchManager::UpdateScanningAssets()
+{
+	bool bTargetState = GetDefault<USearchUserSettings>()->bEnableSearch;
+
+	if (GIsBuildMachine || FApp::IsUnattended())
+	{
+		bTargetState = false;
+	}
+
+	if (bForceEnableSearch)
+	{
+		bTargetState = true;
+	}
+
+	if (bTargetState != bStarted)
+	{
+		bStarted = bTargetState;
+
+		if (bTargetState)
+		{
+			StartScanningAssets();
+		}
+		else
+		{
+			StopScanningAssets();
+		}
+	}
+}
+
+void FAssetSearchManager::StartScanningAssets()
+{
 	IAssetRegistry& AssetRegistry = FModuleManager::LoadModuleChecked<FAssetRegistryModule>("AssetRegistry").Get();
 	AssetRegistry.OnAssetAdded().AddRaw(this, &FAssetSearchManager::OnAssetAdded);
 	AssetRegistry.OnAssetRemoved().AddRaw(this, &FAssetSearchManager::OnAssetRemoved);
@@ -226,9 +239,19 @@ void FAssetSearchManager::InternalStart()
 	{
 		OnAssetAdded(Data);
 	}
+}
 
-	RunThread = true;
-	DatabaseThread = FRunnableThread::Create(this, TEXT("UniversalSearch"), 0, TPri_BelowNormal);
+void FAssetSearchManager::StopScanningAssets()
+{
+	if (FAssetRegistryModule* AssetRegistryModule = FModuleManager::GetModulePtr<FAssetRegistryModule>("AssetRegistry"))
+	{
+		AssetRegistryModule->Get().OnAssetAdded().RemoveAll(this);
+		AssetRegistryModule->Get().OnAssetRemoved().AddRaw(this, &FAssetSearchManager::OnAssetRemoved);
+		AssetRegistryModule->Get().OnFilesLoaded().AddRaw(this, &FAssetSearchManager::OnAssetScanFinished);
+	}
+
+	ProcessAssetQueue.Reset();
+	FailedDDCRequests.Reset();
 }
 
 void FAssetSearchManager::TryConnectToDatabase()
@@ -261,8 +284,8 @@ FSearchStats FAssetSearchManager::GetStats() const
 {
 	FSearchStats Stats;
 	Stats.Scanning = ProcessAssetQueue.Num();
-	Stats.Downloading = PendingDownloads;
-	Stats.PendingDatabaseUpdates = PendingDatabaseUpdates;
+	Stats.Processing = IsAssetUpToDateCount + DownloadQueueCount;
+	Stats.Updating = PendingDatabaseUpdates;
 	Stats.TotalRecords = TotalSearchRecords;
 	Stats.AssetsMissingIndex = FailedDDCRequests.Num();
 	return Stats;
@@ -423,21 +446,34 @@ bool FAssetSearchManager::IsAssetIndexable(UObject* InAsset)
 
 bool FAssetSearchManager::TryLoadIndexForAsset(const FAssetData& InAssetData)
 {
-	return AsyncGetDerivedDataKey(InAssetData, [this, InAssetData](FString InDDCKey) {
+	const bool bSuccess = AsyncGetDerivedDataKey(InAssetData, [this, InAssetData](FString InDDCKey) {
 		FeedOperations.Enqueue([this, InAssetData, InDDCKey]() {
 			FScopeLock ScopedLock(&SearchDatabaseCS);
 			if (!SearchDatabase.IsAssetUpToDate(InAssetData, InDDCKey))
 			{
-				PendingDownloads++;
-
-				FAssetDDCRequest DDCRequest;
-				DDCRequest.AssetData = InAssetData;
-				DDCRequest.DDCKey_IndexDataHash = InDDCKey;
-				DDCRequest.DDCHandle = GetDerivedDataCacheRef().GetAsynchronous(*InDDCKey, InAssetData.ObjectPath.ToString());
-				ProcessDDCQueue.Enqueue(DDCRequest);
+				AsyncRequestDownlaod(InAssetData, InDDCKey);
 			}
+
+			IsAssetUpToDateCount--;
 		});
 	});
+
+	if (bSuccess)
+	{
+		IsAssetUpToDateCount++;
+	}
+
+	return bSuccess;
+}
+
+void FAssetSearchManager::AsyncRequestDownlaod(const FAssetData& InAssetData, const FString& InDDCKey)
+{
+	DownloadQueueCount++;
+
+	FAssetDDCRequest DDCRequest;
+	DDCRequest.AssetData = InAssetData;
+	DDCRequest.DDCKey = InDDCKey;
+	DownloadQueue.Enqueue(DDCRequest);
 }
 
 bool FAssetSearchManager::AsyncGetDerivedDataKey(const FAssetData& InAssetData, TFunction<void(FString)> DDCKeyCallback)
@@ -590,36 +626,17 @@ bool FAssetSearchManager::Tick_GameThread(float DeltaTime)
 {
 	check(IsInGameThread());
 
-	if (bDisableUniversalSearch || !GetDefault<USearchUserSettings>()->bEnableSearch)
-	{
-		return true;
-	}
-
-	if (!bStarted)
-	{
-		InternalStart();
-	}
+	UpdateScanningAssets();
 
 	TryConnectToDatabase();
 
-	//if (0)
-	//{
-	//	IAssetRegistry& AssetRegistry = FModuleManager::LoadModuleChecked<FAssetRegistryModule>("AssetRegistry").Get();
-	//	TArray<FAssetData> TempAssetData;
-	//	AssetRegistry.GetAllAssets(TempAssetData, true);
-
-	//	for (const FAssetData& Data : TempAssetData)
-	//	{
-	//		OnAssetAdded(Data);
-	//	}
-	//}
-
 	ProcessGameThreadTasks();
 
 	const USearchUserSettings* UserSettings = GetDefault<USearchUserSettings>();
+	const FSearchPerformance& PerformanceLimits = bForceEnableSearch ? UserSettings->DefaultOptions : UserSettings->GetPerformanceOptions();
 
-	int32 ScanLimit = GameThread_AssetScanLimit;
-	while (ProcessAssetQueue.Num() > 0 && ScanLimit > 0 && PendingDownloads < PendingDownloadsMax)
+	int32 ScanLimit = PerformanceLimits.AssetScanRate;
+	while (ProcessAssetQueue.Num() > 0 && ScanLimit > 0)
 	{
 		FAssetOperation Operation = ProcessAssetQueue.Pop(false);
 		FAssetData Asset = Operation.Asset;
@@ -635,16 +652,25 @@ bool FAssetSearchManager::Tick_GameThread(float DeltaTime)
 		}
 		else
 		{
-			if (TryLoadIndexForAsset(Asset))
-			{
-				ScanLimit -= 10;
-			}
+			TryLoadIndexForAsset(Asset);
 		}
 
 		ScanLimit--;
 	}
 
-	int32 DownloadProcessLimit = GameThread_DownloadProcessLimit;
+	while (!DownloadQueue.IsEmpty() && ActiveDownloads < PerformanceLimits.ParallelDownloads)
+	{
+		FAssetDDCRequest DDCRequest;
+		bool bSuccess = DownloadQueue.Dequeue(DDCRequest);
+		check(bSuccess);
+
+		DownloadQueueCount--;
+
+		DDCRequest.DDCHandle = GetDerivedDataCacheRef().GetAsynchronous(*DDCRequest.DDCKey, DDCRequest.AssetData.ObjectPath.ToString());
+		ProcessDDCQueue.Enqueue(DDCRequest);
+	}
+
+	int32 DownloadProcessLimit = PerformanceLimits.DownloadProcessRate;
 	while (!ProcessDDCQueue.IsEmpty() && DownloadProcessLimit > 0)
 	{
 		const FAssetDDCRequest* PendingRequest = ProcessDDCQueue.Peek();
@@ -656,15 +682,15 @@ bool FAssetSearchManager::Tick_GameThread(float DeltaTime)
 			bool bGetSuccessful = GetDerivedDataCacheRef().GetAsynchronousResults(PendingRequest->DDCHandle, OutContent, &bDataWasBuilt);
 			if (bGetSuccessful)
 			{
-				LoadDDCContentIntoDatabase(PendingRequest->AssetData, OutContent, PendingRequest->DDCKey_IndexDataHash);
+				LoadDDCContentIntoDatabase(PendingRequest->AssetData, OutContent, PendingRequest->DDCKey);
 			}
-			else
+			else if (UserSettings->bShowAdvancedData)
 			{
 				FailedDDCRequests.Add(*PendingRequest);
 			}
 
 			ProcessDDCQueue.Pop();
-			PendingDownloads--;
+			ActiveDownloads--;
 			DownloadProcessLimit--;
 			continue;
 		}
diff --git a/Engine/Plugins/Editor/AssetSearch/Source/Private/AssetSearchManager.h b/Engine/Plugins/Editor/AssetSearch/Source/Private/AssetSearchManager.h
--- a/Engine/Plugins/Editor/AssetSearch/Source/Private/AssetSearchManager.h
+++ b/Engine/Plugins/Editor/AssetSearch/Source/Private/AssetSearchManager.h
@@ -36,7 +36,9 @@ private:
 	void Tick_DatabaseOperationThread();
 
 private:
-	void InternalStart();
+	void UpdateScanningAssets();
+	void StartScanningAssets();
+	void StopScanningAssets();
 
 	void OnAssetAdded(const FAssetData& InAssetData);
 	void OnAssetRemoved(const FAssetData& InAssetData);
@@ -50,6 +52,7 @@ private:
 	bool RequestIndexAsset(UObject* InAsset);
 	bool IsAssetIndexable(UObject* InAsset);
 	bool TryLoadIndexForAsset(const FAssetData& InAsset);
+	void AsyncRequestDownlaod(const FAssetData& InAssetData, const FString& InDDCKey);
 	bool AsyncGetDerivedDataKey(const FAssetData& UnindexedAsset, TFunction<void(FString)> DDCKeyCallback);
 	bool HasIndexerForClass(const UClass* InAssetClass) const;
 	FString GetIndexerVersion(const UClass* InAssetClass) const;
@@ -68,7 +71,9 @@ private:
 	FAssetSearchDatabase SearchDatabase;
 	FCriticalSection SearchDatabaseCS;
 	TAtomic<int32> PendingDatabaseUpdates;
-	TAtomic<int32> PendingDownloads;
+	TAtomic<int32> IsAssetUpToDateCount;
+	TAtomic<int32> ActiveDownloads;
+	TAtomic<int32> DownloadQueueCount;
 	TAtomic<int64> TotalSearchRecords;
 
 	double LastRecordCountUpdateSeconds;
@@ -89,9 +94,10 @@ private:
 	struct FAssetDDCRequest
 	{
 		FAssetData AssetData;
-		FString DDCKey_IndexDataHash;
+		FString DDCKey;
 		uint32 DDCHandle;
 	};
+	TQueue<FAssetDDCRequest, EQueueMode::Mpsc> DownloadQueue;
 	TQueue<FAssetDDCRequest, EQueueMode::Mpsc> ProcessDDCQueue;
 
 	TArray<FAssetDDCRequest> FailedDDCRequests;
diff --git a/Engine/Plugins/Editor/AssetSearch/Source/Private/Settings/SearchUserSettings.cpp b/Engine/Plugins/Editor/AssetSearch/Source/Private/Settings/SearchUserSettings.cpp
--- a/Engine/Plugins/Editor/AssetSearch/Source/Private/Settings/SearchUserSettings.cpp
+++ b/Engine/Plugins/Editor/AssetSearch/Source/Private/Settings/SearchUserSettings.cpp
@@ -4,4 +4,16 @@
 
 USearchUserSettings::USearchUserSettings()
 {
+	DefaultOptions.ParallelDownloads = 100;
+	DefaultOptions.DownloadProcessRate = 30;
+	DefaultOptions.AssetScanRate = 1000;
+
+	BackgroundtOptions.ParallelDownloads = 1;
+	BackgroundtOptions.DownloadProcessRate = 30;
+	BackgroundtOptions.AssetScanRate = 200;
 }
+
+const FSearchPerformance& USearchUserSettings::GetPerformanceOptions() const
+{
+	return (SearchInForeground > 0) ? DefaultOptions : BackgroundtOptions;
+}
\ No newline at end of file
diff --git a/Engine/Plugins/Editor/AssetSearch/Source/Private/Settings/SearchUserSettings.h b/Engine/Plugins/Editor/AssetSearch/Source/Private/Settings/SearchUserSettings.h
--- a/Engine/Plugins/Editor/AssetSearch/Source/Private/Settings/SearchUserSettings.h
+++ b/Engine/Plugins/Editor/AssetSearch/Source/Private/Settings/SearchUserSettings.h
@@ -6,6 +6,22 @@
 #include "Engine/EngineTypes.h"
 #include "SearchUserSettings.generated.h"
 
+USTRUCT()
+struct FSearchPerformance
+{
+	GENERATED_BODY();
+
+public:
+	UPROPERTY(EditAnywhere, Category = "Performance")
+	int32 ParallelDownloads = 0;
+
+	UPROPERTY(EditAnywhere, AdvancedDisplay, Category = "Performance")
+	int32 DownloadProcessRate = 0;
+
+	UPROPERTY(EditAnywhere, AdvancedDisplay, Category = "Performance")
+	int32 AssetScanRate = 0;
+};
+
 UCLASS(config = EditorPerProjectUserSettings, meta=(DisplayName="Search"))
 class USearchUserSettings : public UDeveloperSettings
 {
@@ -23,4 +39,19 @@ public:
 
 	UPROPERTY(config, EditAnywhere, AdvancedDisplay, Category=General)
 	bool bShowAdvancedData;
+
+	UPROPERTY(config, EditAnywhere, Category = Performance)
+	bool bThrottleInBackground = true;
+
+	UPROPERTY(config, EditAnywhere, Category = Performance)
+	FSearchPerformance DefaultOptions;
+
+	UPROPERTY(config, EditAnywhere, Category = Performance)
+	FSearchPerformance BackgroundtOptions;
+
+	const FSearchPerformance& GetPerformanceOptions() const;
+
+public:
+	UPROPERTY(Transient)
+	int32 SearchInForeground = 0;
 };
diff --git a/Engine/Plugins/Editor/AssetSearch/Source/Private/Widgets/SSearchBrowser.cpp b/Engine/Plugins/Editor/AssetSearch/Source/Private/Widgets/SSearchBrowser.cpp
--- a/Engine/Plugins/Editor/AssetSearch/Source/Private/Widgets/SSearchBrowser.cpp
+++ b/Engine/Plugins/Editor/AssetSearch/Source/Private/Widgets/SSearchBrowser.cpp
@@ -48,14 +48,26 @@ namespace AssetSearchConstants
 
 //TODO Expose TSharedRef<SWidget> SAssetViewItem::CreateToolTipWidget() const via IContentBrowserSingleton.
 
+SSearchBrowser::~SSearchBrowser()
+{
+	if (UObjectInitialized())
+	{
+		GetMutableDefault<USearchUserSettings>()->SearchInForeground--;
+	}
+}
+
 void SSearchBrowser::Construct( const FArguments& InArgs )
 {
-	if (!GetDefault<USearchUserSettings>()->bEnableSearch)
+	USearchUserSettings* UserSettings = GetMutableDefault<USearchUserSettings>();
+
+	if (!UserSettings->bEnableSearch)
 	{
-		GetMutableDefault<USearchUserSettings>()->bEnableSearch = true;
-		GetMutableDefault<USearchUserSettings>()->SaveConfig();
+		UserSettings->bEnableSearch = true;
+		UserSettings->SaveConfig();
 	}
 
+	UserSettings->SearchInForeground++;
+
 	SortByColumn = SSearchTreeRow::NAME_ColumnName;
 	SortMode = EColumnSortMode::Ascending;
 
@@ -155,7 +167,6 @@ void SSearchBrowser::Construct( const FArguments& InArgs )
 					.AutoHeight()
 					[
 						SNew(STextBlock)
-						.Visibility_Lambda([] { return GetDefault<USearchUserSettings>()->bShowAdvancedData ? EVisibility::Visible : EVisibility::Collapsed; })
 						.Text(this, &SSearchBrowser::GetStatusText)
 					]
 
@@ -175,6 +186,7 @@ void SSearchBrowser::Construct( const FArguments& InArgs )
 				[
 					SNew(SHyperlink)
 					.Text(this, &SSearchBrowser::GetUnindexedAssetsText)
+					.Visibility_Lambda([] { return GetDefault<USearchUserSettings>()->bShowAdvancedData ? EVisibility::Visible : EVisibility::Collapsed; })
 					.OnNavigate(this, &SSearchBrowser::HandleForceIndexOfAssetsMissingIndex)
 				]
 			]
@@ -188,7 +200,7 @@ FText SSearchBrowser::GetStatusText() const
 {
 	IAssetSearchModule& SearchModule = IAssetSearchModule::Get();
 	FSearchStats SearchStats = SearchModule.GetStats();
-	int32 UpdatingCount = SearchStats.Scanning + SearchStats.Downloading + SearchStats.PendingDatabaseUpdates;
+	int32 UpdatingCount = SearchStats.Scanning + SearchStats.Processing + SearchStats.Updating;
 
 	if (UpdatingCount > 0)
 	{
@@ -204,7 +216,7 @@ FText SSearchBrowser::GetAdvancedStatus() const
 {
 	IAssetSearchModule& SearchModule = IAssetSearchModule::Get();
 	FSearchStats SearchStats = SearchModule.GetStats();
-	return FText::Format(LOCTEXT("AdvancedSearchStatusTextFmt", "Scanning {0}   Downloading {1}   Updating {2}            Total Records {3}"), SearchStats.Scanning, SearchStats.Downloading, SearchStats.PendingDatabaseUpdates, SearchStats.TotalRecords);
+	return FText::Format(LOCTEXT("AdvancedSearchStatusTextFmt", "Scanning {0}   Processing {1}   Updating {2}            Total Records {3}"), SearchStats.Scanning, SearchStats.Processing, SearchStats.Updating, SearchStats.TotalRecords);
 }
 
 FText SSearchBrowser::GetUnindexedAssetsText() const
diff --git a/Engine/Plugins/Editor/AssetSearch/Source/Private/Widgets/SSearchBrowser.h b/Engine/Plugins/Editor/AssetSearch/Source/Private/Widgets/SSearchBrowser.h
--- a/Engine/Plugins/Editor/AssetSearch/Source/Private/Widgets/SSearchBrowser.h
+++ b/Engine/Plugins/Editor/AssetSearch/Source/Private/Widgets/SSearchBrowser.h
@@ -29,6 +29,8 @@ public:
 
 public:
 
+	virtual ~SSearchBrowser();
+
 	/**
 	 * Construct this widget
 	 *
diff --git a/Engine/Plugins/Editor/AssetSearch/Source/Public/IAssetSearchModule.h b/Engine/Plugins/Editor/AssetSearch/Source/Public/IAssetSearchModule.h
--- a/Engine/Plugins/Editor/AssetSearch/Source/Public/IAssetSearchModule.h
+++ b/Engine/Plugins/Editor/AssetSearch/Source/Public/IAssetSearchModule.h
@@ -41,8 +41,8 @@ struct FSearchRecord
 struct FSearchStats
 {
 	int32 Scanning = 0;
-	int32 Downloading = 0;
-	int32 PendingDatabaseUpdates = 0;
+	int32 Processing = 0;
+	int32 Updating = 0;
 
 	int32 AssetsMissingIndex = 0;
 
