Universal Search - Working on better search feedback, and handling being resized better.  Removing several disable optimizations from when testing.

diff --git a/Engine/Plugins/Editor/AssetSearch/Source/AssetSearch.Build.cs b/Engine/Plugins/Editor/AssetSearch/Source/AssetSearch.Build.cs
--- a/Engine/Plugins/Editor/AssetSearch/Source/AssetSearch.Build.cs
+++ b/Engine/Plugins/Editor/AssetSearch/Source/AssetSearch.Build.cs
@@ -33,7 +33,7 @@ public class AssetSearch : ModuleRules
 				"Json",
 				"AssetRegistry",
 				"JsonUtilities",
-
+				"Projects",
 				"UMG",
 				"UMGEditor",
 				"BlueprintGraph",
diff --git a/Engine/Plugins/Editor/AssetSearch/Source/Private/AssetSearchDatabase.cpp b/Engine/Plugins/Editor/AssetSearch/Source/Private/AssetSearchDatabase.cpp
--- a/Engine/Plugins/Editor/AssetSearch/Source/Private/AssetSearchDatabase.cpp
+++ b/Engine/Plugins/Editor/AssetSearch/Source/Private/AssetSearchDatabase.cpp
@@ -17,8 +17,6 @@
 #include "HAL/PlatformFilemanager.h"
 #include "Misc/StringBuilder.h"
 
-PRAGMA_DISABLE_OPTIMIZATION
-
 DECLARE_LOG_CATEGORY_CLASS(LogAssetSearch, Log, All);
 
 enum class EAssetSearchDatabaseVersion
@@ -846,5 +844,3 @@ FString FSearchQuery::ConvertToDatabaseQuery() const
 
 	return Q.ToString();
 }
-
-PRAGMA_ENABLE_OPTIMIZATION
diff --git a/Engine/Plugins/Editor/AssetSearch/Source/Private/AssetSearchManager.cpp b/Engine/Plugins/Editor/AssetSearch/Source/Private/AssetSearchManager.cpp
--- a/Engine/Plugins/Editor/AssetSearch/Source/Private/AssetSearchManager.cpp
+++ b/Engine/Plugins/Editor/AssetSearch/Source/Private/AssetSearchManager.cpp
@@ -37,8 +37,8 @@
 #include "UObject/ObjectMacros.h"
 #include "UObject/ObjectKey.h"
 #include "UObject/WeakObjectPtrTemplates.h"
-
-PRAGMA_DISABLE_OPTIMIZATION
+#include "FileHelpers.h"
+#include "Misc/MessageDialog.h"
 
 #define LOCTEXT_NAMESPACE "FAssetSearchManager"
 
@@ -684,7 +684,7 @@ bool FAssetSearchManager::Tick_GameThread(float DeltaTime)
 			{
 				LoadDDCContentIntoDatabase(PendingRequest->AssetData, OutContent, PendingRequest->DDCKey);
 			}
-			else if (UserSettings->bShowAdvancedData)
+			else if (UserSettings->bShowMissingAssets)
 			{
 				FailedDDCRequests.Add(*PendingRequest);
 			}
@@ -747,6 +747,8 @@ void FAssetSearchManager::ForceIndexOnAssetsMissingIndex()
 
 	int32 RemovedCount = 0;
 
+	TArray<FAssetData> RedirectorsWithBrokenMetadata;
+
 	FUnloadPackageScope UnloadScope;
 	for (const FAssetDDCRequest& Request : FailedDDCRequests)
 	{
@@ -758,6 +760,15 @@ void FAssetSearchManager::ForceIndexOnAssetsMissingIndex()
 		IndexingTask.EnterProgressFrame(1, FText::Format(LOCTEXT("ForceIndexOnAssetsMissingIndexFormat", "Indexing Asset ({0} of {1})"), RemovedCount + 1, FailedDDCRequests.Num()));
 		if (UObject* AssetToIndex = Request.AssetData.GetAsset())
 		{
+			// This object's metadata incorrectly labled it as something other than a redirector.  We need to resave it
+			// to stop it from appearing as something it's not.
+			if (UObjectRedirector* Redirector = Cast<UObjectRedirector>(AssetToIndex))
+			{
+				RedirectorsWithBrokenMetadata.Add(Request.AssetData);
+				RemovedCount++;
+				continue;
+			}
+
 			StoreIndexForAsset(AssetToIndex);
 		}
 
@@ -769,6 +780,26 @@ void FAssetSearchManager::ForceIndexOnAssetsMissingIndex()
 		RemovedCount++;
 	}
 
+	if (RedirectorsWithBrokenMetadata.Num() > 0)
+	{
+		EAppReturnType::Type ReturnValue = FMessageDialog::Open(EAppMsgType::YesNo,
+			LOCTEXT("ResaveRedirectors", "We found some redirectors that didn't have the correct asset metadata identifying them as redirectors.  Would you like to resave them, so that they stop appearing as missing asset indexes?"));
+
+		if (ReturnValue == EAppReturnType::Yes)
+		{
+			TArray<UPackage*> PackagesToSave;
+			for (const FAssetData& BrokenAsset : RedirectorsWithBrokenMetadata)
+			{
+				if (UObject* Redirector = BrokenAsset.GetAsset())
+				{
+					PackagesToSave.Add(Redirector->GetOutermost());
+				}
+			}
+
+			FEditorFileUtils::PromptForCheckoutAndSave(PackagesToSave, /*bCheckDirty*/false, /*bPromptToSave*/false);
+		}
+	}
+
 	FailedDDCRequests.RemoveAtSwap(0, RemovedCount);
 }
 
@@ -826,5 +857,3 @@ void FAssetSearchManager::ProcessGameThreadTasks()
 }
 
 #undef LOCTEXT_NAMESPACE
-
-PRAGMA_ENABLE_OPTIMIZATION
\ No newline at end of file
diff --git a/Engine/Plugins/Editor/AssetSearch/Source/Private/AssetSearchModule.cpp b/Engine/Plugins/Editor/AssetSearch/Source/Private/AssetSearchModule.cpp
--- a/Engine/Plugins/Editor/AssetSearch/Source/Private/AssetSearchModule.cpp
+++ b/Engine/Plugins/Editor/AssetSearch/Source/Private/AssetSearchModule.cpp
@@ -14,8 +14,6 @@
 
 #define LOCTEXT_NAMESPACE "FAssetSearchModule"
 
-PRAGMA_DISABLE_OPTIMIZATION
-
 static const FName SearchTabName("Search");
 
 class FAssetSearchModule : public IAssetSearchModule
@@ -87,6 +85,4 @@ private:
 
 #undef LOCTEXT_NAMESPACE
 
-PRAGMA_ENABLE_OPTIMIZATION
-
 IMPLEMENT_MODULE(FAssetSearchModule, AssetSearch);
diff --git a/Engine/Plugins/Editor/AssetSearch/Source/Private/FileInfoDatabase.cpp b/Engine/Plugins/Editor/AssetSearch/Source/Private/FileInfoDatabase.cpp
--- a/Engine/Plugins/Editor/AssetSearch/Source/Private/FileInfoDatabase.cpp
+++ b/Engine/Plugins/Editor/AssetSearch/Source/Private/FileInfoDatabase.cpp
@@ -17,8 +17,6 @@
 #include "HAL/PlatformFilemanager.h"
 #include "Misc/StringBuilder.h"
 
-PRAGMA_DISABLE_OPTIMIZATION
-
 DECLARE_LOG_CATEGORY_CLASS(LogFileInfo, Log, All);
 
 enum class EFileInfoDatabaseVersion
@@ -432,5 +430,3 @@ TMap<FName, FAssetFileInfo> FFileInfoDatabase::GetAllFileInfos()
 
 	return FileInfos;
 }
-
-PRAGMA_ENABLE_OPTIMIZATION
diff --git a/Engine/Plugins/Editor/AssetSearch/Source/Private/Settings/SearchUserSettings.h b/Engine/Plugins/Editor/AssetSearch/Source/Private/Settings/SearchUserSettings.h
--- a/Engine/Plugins/Editor/AssetSearch/Source/Private/Settings/SearchUserSettings.h
+++ b/Engine/Plugins/Editor/AssetSearch/Source/Private/Settings/SearchUserSettings.h
@@ -37,8 +37,8 @@ public:
 	UPROPERTY(config, EditAnywhere, Category=General)
 	TArray<FDirectoryPath> IgnoredPaths;
 
-	UPROPERTY(config, EditAnywhere, AdvancedDisplay, Category=General)
-	bool bShowAdvancedData;
+	UPROPERTY(config, EditAnywhere, Category=General)
+	bool bShowMissingAssets;
 
 	UPROPERTY(config, EditAnywhere, Category = Performance)
 	bool bThrottleInBackground = true;
diff --git a/Engine/Plugins/Editor/AssetSearch/Source/Private/Widgets/SSearchBrowser.cpp b/Engine/Plugins/Editor/AssetSearch/Source/Private/Widgets/SSearchBrowser.cpp
--- a/Engine/Plugins/Editor/AssetSearch/Source/Private/Widgets/SSearchBrowser.cpp
+++ b/Engine/Plugins/Editor/AssetSearch/Source/Private/Widgets/SSearchBrowser.cpp
@@ -35,6 +35,8 @@
 #include "Widgets/Images/SThrobber.h"
 #include "Widgets/Input/SSearchBox.h"
 #include "Settings/SearchUserSettings.h"
+#include "SearchStyle.h"
+#include "Widgets/SToolTip.h"
 
 #define LOCTEXT_NAMESPACE "SObjectBrowser"
 
@@ -46,6 +48,7 @@ namespace AssetSearchConstants
 	const int32 ThumbnailPoolSize = 64;
 }
 
+//;
 //TODO Expose TSharedRef<SWidget> SAssetViewItem::CreateToolTipWidget() const via IContentBrowserSingleton.
 
 SSearchBrowser::~SSearchBrowser()
@@ -89,17 +92,17 @@ void SSearchBrowser::Construct( const FArguments& InArgs )
 			[
 				SNew(SHorizontalBox)
 
-				+ SHorizontalBox::Slot()
-				.AutoWidth()
-				[
-					SNew(SButton)
-					.OnClicked(this, &SSearchBrowser::OnRefresh)
-					[
-						SNew(STextBlock)
-						.Font(FEditorStyle::Get().GetFontStyle("FontAwesome.11"))
-						.Text(FEditorFontGlyphs::Refresh)
-					]
-				]
+				//+ SHorizontalBox::Slot()
+				//.AutoWidth()
+				//[
+				//	SNew(SButton)
+				//	.OnClicked(this, &SSearchBrowser::OnRefresh)
+				//	[
+				//		SNew(STextBlock)
+				//		.Font(FEditorStyle::Get().GetFontStyle("FontAwesome.11"))
+				//		.Text(FEditorFontGlyphs::Refresh)
+				//	]
+				//]
 
 				+ SHorizontalBox::Slot()
 				.FillWidth(1.0f)
@@ -120,25 +123,40 @@ void SSearchBrowser::Construct( const FArguments& InArgs )
 				.BorderImage(FEditorStyle::GetBrush("ToolPanel.GroupBorder"))
 				.Padding(FMargin(0.0f, 4.0f))
 				[
-					SAssignNew(SearchTreeView, STreeView< TSharedPtr<FSearchNode> >)
-					.ItemHeight(24.0f)
-					.TreeItemsSource(&SearchResults)
-					.SelectionMode(ESelectionMode::Single)
-					.OnGenerateRow(this, &SSearchBrowser::HandleListGenerateRow)
-					.OnGetChildren(this, &SSearchBrowser::GetChildrenForInfo)
-					.OnSelectionChanged(this, &SSearchBrowser::HandleListSelectionChanged)
-					.HeaderRow
-					(
-						SNew(SHeaderRow)
-
-						+ SHeaderRow::Column(SSearchTreeRow::NAME_ColumnName)
-						.DefaultLabel(LOCTEXT("ColumnName", "Name"))
-						.FillWidth(0.80f)
-
-						+ SHeaderRow::Column(SSearchTreeRow::NAME_ColumnType)
-						.ManualWidth(300)
-						.DefaultLabel(LOCTEXT("ColumnType", "Type"))
-					)
+					SNew(SOverlay)
+
+					+ SOverlay::Slot()
+					.HAlign(HAlign_Center)
+					.VAlign(VAlign_Center)
+					[
+						SNew(STextBlock)
+						.Clipping(EWidgetClipping::Inherit)
+						.Text(this, &SSearchBrowser::GetSearchBackgroundText)
+						.Justification(ETextJustify::Center)
+						.Font(FCoreStyle::GetDefaultFontStyle("Regular", 30))
+						.ColorAndOpacity(FLinearColor(1,1,1,0.05))
+						.RenderTransformPivot(FVector2D(0.5, 0.5))
+						.RenderTransform(FSlateRenderTransform(FQuat2D(FMath::DegreesToRadians(-30.0f))))
+					]
+
+					+ SOverlay::Slot()
+					[
+						SAssignNew(SearchTreeView, STreeView< TSharedPtr<FSearchNode> >)
+						.ItemHeight(24.0f)
+						.TreeItemsSource(&SearchResults)
+						.SelectionMode(ESelectionMode::Single)
+						.OnGenerateRow(this, &SSearchBrowser::HandleListGenerateRow)
+						.OnGetChildren(this, &SSearchBrowser::GetChildrenForInfo)
+						.OnSelectionChanged(this, &SSearchBrowser::HandleListSelectionChanged)
+						.HeaderRow
+						(
+							SAssignNew(HeaderColumns, SHeaderRow)
+
+							+ SHeaderRow::Column(SSearchTreeRow::NAME_ColumnName)
+							.DefaultLabel(LOCTEXT("ColumnName", "Name"))
+							.FillWidth(70)
+						)
+					]
 				]
 			]
 
@@ -155,28 +173,36 @@ void SSearchBrowser::Construct( const FArguments& InArgs )
 			[
 				SNew(SHorizontalBox)
 
-				// Asset Stats 
 				+ SHorizontalBox::Slot()
-				.FillWidth(1.f)
+				.AutoWidth()
 				.VAlign(VAlign_Center)
 				.Padding(8, 0)
 				[
-					SNew(SVerticalBox)
-
-					+ SVerticalBox::Slot()
-					.AutoHeight()
-					[
-						SNew(STextBlock)
-						.Text(this, &SSearchBrowser::GetStatusText)
-					]
+					SNew(SImage)
+					.Image(FSearchStyle::Get().GetBrush("Stats"))
+					.ToolTip(
+						SNew(SToolTip)
+						[
+							SNew(SVerticalBox)
+
+							+ SVerticalBox::Slot()
+							.AutoHeight()
+							[
+								SNew(STextBlock)
+								.Text(this, &SSearchBrowser::GetAdvancedStatus)
+							]
+						]
+					)
+				]
 
-					+ SVerticalBox::Slot()
-					.AutoHeight()
-					[
-						SNew(STextBlock)
-						.Visibility_Lambda([] { return GetDefault<USearchUserSettings>()->bShowAdvancedData ? EVisibility::Visible : EVisibility::Collapsed; })
-						.Text(this, &SSearchBrowser::GetAdvancedStatus)
-					]
+				// Asset Stats 
+				+ SHorizontalBox::Slot()
+				.FillWidth(1.f)
+				.VAlign(VAlign_Center)
+				.Padding(2, 0)
+				[
+					SNew(STextBlock)
+					.Text(this, &SSearchBrowser::GetStatusText)
 				]
 
 				// Index unindexed items
@@ -186,7 +212,7 @@ void SSearchBrowser::Construct( const FArguments& InArgs )
 				[
 					SNew(SHyperlink)
 					.Text(this, &SSearchBrowser::GetUnindexedAssetsText)
-					.Visibility_Lambda([] { return GetDefault<USearchUserSettings>()->bShowAdvancedData ? EVisibility::Visible : EVisibility::Collapsed; })
+					.Visibility_Lambda([] { return GetDefault<USearchUserSettings>()->bShowMissingAssets ? EVisibility::Visible : EVisibility::Collapsed; })
 					.OnNavigate(this, &SSearchBrowser::HandleForceIndexOfAssetsMissingIndex)
 				]
 			]
@@ -196,6 +222,41 @@ void SSearchBrowser::Construct( const FArguments& InArgs )
 	RefreshList();
 }
 
+void SSearchBrowser::Tick(const FGeometry& AllottedGeometry, const double InCurrentTime, const float InDeltaTime)
+{
+	SCompoundWidget::Tick(AllottedGeometry, InCurrentTime, InDeltaTime);
+
+	const int32 MinSizeForTypeColumn = 400;
+	if (HeaderColumns->GetColumns().Num() == 1 && AllottedGeometry.GetLocalSize().X > MinSizeForTypeColumn)
+	{
+		HeaderColumns->AddColumn(
+			SHeaderRow::Column(SSearchTreeRow::NAME_ColumnType)
+			.DefaultLabel(LOCTEXT("ColumnType", "Type"))
+			.FillWidth(30)
+		);
+	}
+	else if (HeaderColumns->GetColumns().Num() > 1 && AllottedGeometry.GetLocalSize().X < MinSizeForTypeColumn)
+	{
+		HeaderColumns->RemoveColumn(SSearchTreeRow::NAME_ColumnType);
+	}
+}
+
+FText SSearchBrowser::GetSearchBackgroundText() const
+{
+	if (FilterString.Len() > 0 && !IsSearching() && SearchResults.Num() == 0)
+	{
+		return LOCTEXT("SearchAllTheThings", "¯\\_(ツ)_/¯");
+	}
+	else if (FilterString.Len() == 0)
+	{
+		return LOCTEXT("SearchAllTheThings", "Search\nAll The Things!");
+	}
+	else
+	{
+		return FText::GetEmpty();
+	}
+}
+
 FText SSearchBrowser::GetStatusText() const
 {
 	IAssetSearchModule& SearchModule = IAssetSearchModule::Get();
@@ -216,14 +277,20 @@ FText SSearchBrowser::GetAdvancedStatus() const
 {
 	IAssetSearchModule& SearchModule = IAssetSearchModule::Get();
 	FSearchStats SearchStats = SearchModule.GetStats();
-	return FText::Format(LOCTEXT("AdvancedSearchStatusTextFmt", "Scanning {0}   Processing {1}   Updating {2}            Total Records {3}"), SearchStats.Scanning, SearchStats.Processing, SearchStats.Updating, SearchStats.TotalRecords);
+	return FText::Format(LOCTEXT("AdvancedSearchStatusTextFmt", "Scanning {0}\nProcessing {1}\nUpdating {2}\n\nTotal Records {3}"), SearchStats.Scanning, SearchStats.Processing, SearchStats.Updating, SearchStats.TotalRecords);
 }
 
 FText SSearchBrowser::GetUnindexedAssetsText() const
 {
 	IAssetSearchModule& SearchModule = IAssetSearchModule::Get();
 	FSearchStats SearchStats = SearchModule.GetStats();
-	return FText::Format(LOCTEXT("UnindexedAssetsLinkFormat", "{0} Missing"), SearchStats.AssetsMissingIndex);
+
+	if (SearchStats.AssetsMissingIndex > 0)
+	{
+		return FText::Format(LOCTEXT("UnindexedAssetsLinkFormat", "{0} Missing"), SearchStats.AssetsMissingIndex);
+	}
+
+	return FText::GetEmpty();
 }
 
 void SSearchBrowser::HandleForceIndexOfAssetsMissingIndex()
diff --git a/Engine/Plugins/Editor/AssetSearch/Source/Private/Widgets/SSearchBrowser.h b/Engine/Plugins/Editor/AssetSearch/Source/Private/Widgets/SSearchBrowser.h
--- a/Engine/Plugins/Editor/AssetSearch/Source/Private/Widgets/SSearchBrowser.h
+++ b/Engine/Plugins/Editor/AssetSearch/Source/Private/Widgets/SSearchBrowser.h
@@ -16,6 +16,7 @@ class FSearchNode;
 class FAssetNode;
 class FMenuBuilder;
 class IAssetRegistry;
+class SHeaderRow;
 
 /**
  * Implements the undo history panel.
@@ -39,9 +40,11 @@ public:
 	void Construct( const FArguments& InArgs );
 
 public:
+	virtual void Tick(const FGeometry& AllottedGeometry, const double InCurrentTime, const float InDeltaTime) override;
 
 private:
 
+	FText GetSearchBackgroundText() const;
 	FText GetStatusText() const;
 	FText GetAdvancedStatus() const;
 	FText GetUnindexedAssetsText() const;
@@ -85,6 +88,8 @@ private:
 
 	TSharedPtr<FAssetThumbnailPool> ThumbnailPool;
 
+	TSharedPtr<SHeaderRow> HeaderColumns;
+
 	FName SortByColumn;
 	EColumnSortMode::Type SortMode;
 };
diff --git a/Engine/Plugins/Editor/AssetSearch/Source/Private/Widgets/SearchStyle.cpp b/Engine/Plugins/Editor/AssetSearch/Source/Private/Widgets/SearchStyle.cpp
new file mode 100644
--- /dev/null
+++ b/Engine/Plugins/Editor/AssetSearch/Source/Private/Widgets/SearchStyle.cpp
@@ -0,0 +1,45 @@
+// Copyright Epic Games, Inc. All Rights Reserved.
+
+#include "SearchStyle.h"
+
+#include "Styling/SlateTypes.h"
+#include "Styling/CoreStyle.h"
+#include "EditorStyleSet.h"
+#include "Interfaces/IPluginManager.h"
+
+FPluginStyleSet::FPluginStyleSet(const FName& InPluginName, const FName& InStyleSetName)
+	: FSlateStyleSet(InStyleSetName == NAME_None ? InPluginName : FName(InPluginName.ToString() + TEXT(".") + InStyleSetName.ToString()))
+{
+	const TSharedPtr<IPlugin> Plugin = IPluginManager::Get().FindPlugin(InPluginName.ToString());
+	check(Plugin.IsValid());
+
+	SetContentRoot(Plugin->GetContentDir());
+	SetCoreContentRoot(FPaths::EngineContentDir() / TEXT("Slate"));
+}
+
+const FVector2D Icon16x16(16.0f, 16.0f);
+const FVector2D Icon20x20(20.0f, 20.0f);
+const FVector2D Icon40x40(40.0f, 40.0f);
+
+FSearchStyle::FSearchStyle()
+	: FPluginStyleSet(TEXT("AssetSearch"))
+{
+	const FVector2D IconSize(20.0f, 20.0f);
+	const FVector2D SmallIconSize(20.0f, 20.0f);
+	const FVector2D LabelIconSize(16.0f, 16.0f);
+
+	Set("Stats", new FSlateImageBrush(RootToContentDir(TEXT("icon_tab_Stats_16x.png")), Icon16x16));
+
+	FSlateStyleRegistry::RegisterSlateStyle(*this);
+}
+
+FSearchStyle::~FSearchStyle()
+{
+	FSlateStyleRegistry::UnRegisterSlateStyle(*this);
+}
+
+FSearchStyle& FSearchStyle::Get()
+{
+	static FSearchStyle Inst;
+	return Inst;
+}
diff --git a/Engine/Plugins/Editor/AssetSearch/Source/Private/Widgets/SearchStyle.h b/Engine/Plugins/Editor/AssetSearch/Source/Private/Widgets/SearchStyle.h
new file mode 100644
--- /dev/null
+++ b/Engine/Plugins/Editor/AssetSearch/Source/Private/Widgets/SearchStyle.h
@@ -0,0 +1,31 @@
+// Copyright Epic Games, Inc. All Rights Reserved.
+
+#pragma once
+
+#include "Styling/SlateStyle.h"
+#include "Styling/SlateStyleRegistry.h"
+
+/**
+ * 
+ */
+class FPluginStyleSet : public FSlateStyleSet
+{
+public:
+	FPluginStyleSet(const FName& InPluginName, const FName& InStyleSetName = NAME_None);
+};
+
+/**
+ * 
+ */
+class FSearchStyle : public FPluginStyleSet
+{
+public:
+
+	/** Access the singleton instance for this style set */
+	static FSearchStyle& Get();
+
+private:
+
+	FSearchStyle();
+	~FSearchStyle();
+};
