Universnal Search - Hashing files after package save completes instead of on presave, was causing the older version to be hashed for the newly parsed object.

diff --git a/Engine/Plugins/Editor/AssetSearch/Source/Private/AssetSearchManager.cpp b/Engine/Plugins/Editor/AssetSearch/Source/Private/AssetSearchManager.cpp
--- a/Engine/Plugins/Editor/AssetSearch/Source/Private/AssetSearchManager.cpp
+++ b/Engine/Plugins/Editor/AssetSearch/Source/Private/AssetSearchManager.cpp
@@ -170,7 +170,7 @@ FAssetSearchManager::~FAssetSearchManager()
 
 	StopScanningAssets();
 
-	FCoreUObjectDelegates::OnObjectSaved.RemoveAll(this);
+	UPackage::PackageSavedEvent.RemoveAll(this);
 	FCoreUObjectDelegates::OnAssetLoaded.RemoveAll(this);
 	UObject::FAssetRegistryTag::OnGetExtraObjectTags.RemoveAll(this);
 
@@ -187,7 +187,7 @@ void FAssetSearchManager::Start()
 	RegisterAssetIndexer(UWorld::StaticClass(), MakeUnique<FLevelIndexer>());
 	RegisterAssetIndexer(USoundCue::StaticClass(), MakeUnique<FSoundCueIndexer>());
 
-	FCoreUObjectDelegates::OnObjectSaved.AddRaw(this, &FAssetSearchManager::OnObjectSaved);
+	UPackage::PackageSavedEvent.AddRaw(this, &FAssetSearchManager::HandlePackageSaved);
 	FCoreUObjectDelegates::OnAssetLoaded.AddRaw(this, &FAssetSearchManager::OnAssetLoaded);
 
 	TickerHandle = FTicker::GetCoreTicker().AddTicker(FTickerDelegate::CreateRaw(this, &FAssetSearchManager::Tick_GameThread), 0);
@@ -372,13 +372,27 @@ void FAssetSearchManager::OnAssetScanFinished()
 	});
 }
 
-void FAssetSearchManager::OnObjectSaved(UObject* InObject)
+void FAssetSearchManager::HandlePackageSaved(const FString& PackageFilename, UObject* Outer)
 {
 	check(IsInGameThread());
 
-	if (!GIsCookerLoadingPackage)
+	// Ignore package operations fired by the cooker (cook on the fly).
+	if (GIsCookerLoadingPackage)
 	{
-		RequestIndexAsset(InObject);
+		return;
+	}
+
+	UPackage* Package = CastChecked<UPackage>(Outer);
+
+	if (GIsEditor && !IsRunningCommandlet())
+	{
+		TArray<UObject*> Objects;
+		const bool bIncludeNestedObjects = false;
+		GetObjectsWithOuter(Package, Objects, bIncludeNestedObjects);
+		for (UObject* Entry : Objects)
+		{
+			RequestIndexAsset(Entry);
+		}
 	}
 }
 
diff --git a/Engine/Plugins/Editor/AssetSearch/Source/Private/AssetSearchManager.h b/Engine/Plugins/Editor/AssetSearch/Source/Private/AssetSearchManager.h
--- a/Engine/Plugins/Editor/AssetSearch/Source/Private/AssetSearchManager.h
+++ b/Engine/Plugins/Editor/AssetSearch/Source/Private/AssetSearchManager.h
@@ -44,7 +44,7 @@ private:
 	void OnAssetRemoved(const FAssetData& InAssetData);
 	void OnAssetScanFinished();
 
-	void OnObjectSaved(UObject* InObject);
+	void HandlePackageSaved(const FString& PackageFilename, UObject* Outer);
 	void OnAssetLoaded(UObject* InObject);
 
 	void AddOrUpdateAsset(const FAssetData& InAsset, const FString& IndexedJson, const FString& DerivedDataKey);
diff --git a/Engine/Plugins/Editor/AssetSearch/Source/Private/Settings/SearchUserSettings.h b/Engine/Plugins/Editor/AssetSearch/Source/Private/Settings/SearchUserSettings.h
--- a/Engine/Plugins/Editor/AssetSearch/Source/Private/Settings/SearchUserSettings.h
+++ b/Engine/Plugins/Editor/AssetSearch/Source/Private/Settings/SearchUserSettings.h
@@ -38,7 +38,7 @@ public:
 	TArray<FDirectoryPath> IgnoredPaths;
 
 	UPROPERTY(config, EditAnywhere, Category=General)
-	bool bShowMissingAssets;
+	bool bShowMissingAssets = true;
 
 	UPROPERTY(config, EditAnywhere, Category = Performance)
 	bool bThrottleInBackground = true;
