Build fixes.

diff --git a/Engine/Plugins/Editor/AssetSearch/Source/Private/Widgets/SSearchBrowser.cpp b/Engine/Plugins/Editor/AssetSearch/Source/Private/Widgets/SSearchBrowser.cpp
--- a/Engine/Plugins/Editor/AssetSearch/Source/Private/Widgets/SSearchBrowser.cpp
+++ b/Engine/Plugins/Editor/AssetSearch/Source/Private/Widgets/SSearchBrowser.cpp
@@ -37,6 +37,7 @@
 #include "Settings/SearchUserSettings.h"
 #include "SearchStyle.h"
 #include "Widgets/SToolTip.h"
+#include "Widgets/Images/SImage.h"
 
 #define LOCTEXT_NAMESPACE "SObjectBrowser"
 
diff --git a/Engine/Plugins/Editor/AssetSearch/Source/Private/Widgets/SearchStyle.cpp b/Engine/Plugins/Editor/AssetSearch/Source/Private/Widgets/SearchStyle.cpp
--- a/Engine/Plugins/Editor/AssetSearch/Source/Private/Widgets/SearchStyle.cpp
+++ b/Engine/Plugins/Editor/AssetSearch/Source/Private/Widgets/SearchStyle.cpp
@@ -8,7 +8,7 @@
 #include "Interfaces/IPluginManager.h"
 
 FPluginStyleSet::FPluginStyleSet(const FName& InPluginName, const FName& InStyleSetName)
-	: FSlateStyleSet(InStyleSetName == NAME_None ? InPluginName : FName(InPluginName.ToString() + TEXT(".") + InStyleSetName.ToString()))
+	: FSlateStyleSet(InStyleSetName == NAME_None ? InPluginName : FName(*(InPluginName.ToString() + TEXT(".") + InStyleSetName.ToString())))
 {
 	const TSharedPtr<IPlugin> Plugin = IPluginManager::Get().FindPlugin(InPluginName.ToString());
 	check(Plugin.IsValid());
