Unversial Search - Allowing the asset names themselves to be searchable.  Adding a view options button to allowing more options in the future for how we view the search results.

diff --git a/Engine/Plugins/Editor/AssetSearch/Source/Private/AssetSearchDatabase.cpp b/Engine/Plugins/Editor/AssetSearch/Source/Private/AssetSearchDatabase.cpp
--- a/Engine/Plugins/Editor/AssetSearch/Source/Private/AssetSearchDatabase.cpp
+++ b/Engine/Plugins/Editor/AssetSearch/Source/Private/AssetSearchDatabase.cpp
@@ -64,7 +64,8 @@ public:
 		PREPARE_STATEMENT(Statement_AddAssetProperty);
 		PREPARE_STATEMENT(Statement_DeleteEntriesForAsset);
 
-		PREPARE_STATEMENT(Statement_SearchAssetsFTS);
+		PREPARE_STATEMENT(Statement_SearchAssets);
+		PREPARE_STATEMENT(Statement_SearchAssetPropertiesFTS);
 
 #undef PREPARE_STATEMENT
 
@@ -292,7 +293,36 @@ public:
 		return false;
 	}
 
-	SQLITE_PREPARED_STATEMENT(FSearchAssetsFTS,
+	SQLITE_PREPARED_STATEMENT(FSearchAssets,
+		" SELECT "
+		"     asset_name, "
+		"     asset_class, "
+		"     asset_path "
+		" FROM table_assets "
+		" WHERE asset_name LIKE ?1 "
+		";",
+		SQLITE_PREPARED_STATEMENT_COLUMNS(FString, FString, FString),
+		SQLITE_PREPARED_STATEMENT_BINDINGS(FString)
+	);
+	FSearchAssets Statement_SearchAssets;
+	bool SearchAssets(const FSearchQuery& Query, TFunctionRef<ESQLitePreparedStatementExecuteRowResult(FSearchRecord&&)> InCallback)
+	{
+		const FString Q = TEXT("%") + Query.Query + TEXT("%");
+
+		return Statement_SearchAssets.BindAndExecute(Q, [&InCallback](const FSearchAssets& InStatement)
+		{
+			FSearchRecord Result;
+			if (InStatement.GetColumnValues(Result.AssetName, Result.AssetClass, Result.AssetPath))
+			{
+				Result.Score = -30;
+
+				return InCallback(MoveTemp(Result));
+			}
+			return ESQLitePreparedStatementExecuteRowResult::Error;
+		}) != INDEX_NONE;
+	}
+
+	SQLITE_PREPARED_STATEMENT(FSearchAssetPropertiesFTS,
 		" SELECT "
 		"     asset_name, "
 		"     asset_class, "
@@ -308,17 +338,16 @@ public:
 		"     rank as score "
 		" FROM table_asset_properties_fts "
 		" WHERE table_asset_properties_fts MATCH ?1 "
-		//" ORDER BY rank "
 		";",
 		SQLITE_PREPARED_STATEMENT_COLUMNS(FString, FString, FString, FString, FString, FString, FString, FString, FString, FString, FString, float),
 		SQLITE_PREPARED_STATEMENT_BINDINGS(FString)
 	);
-	FSearchAssetsFTS Statement_SearchAssetsFTS;
-	bool SearchAssets(const FSearchQuery& Query, TFunctionRef<ESQLitePreparedStatementExecuteRowResult(FSearchRecord&&)> InCallback)
+	FSearchAssetPropertiesFTS Statement_SearchAssetPropertiesFTS;
+	bool SearchAssetProperties(const FSearchQuery& Query, TFunctionRef<ESQLitePreparedStatementExecuteRowResult(FSearchRecord&&)> InCallback)
 	{
 		const FString Q = Query.ConvertToDatabaseQuery();
 
-		return Statement_SearchAssetsFTS.BindAndExecute(Q, [&InCallback](const FSearchAssetsFTS& InStatement)
+		return Statement_SearchAssetPropertiesFTS.BindAndExecute(Q, [&InCallback](const FSearchAssetPropertiesFTS& InStatement)
 		{
 			FSearchRecord Result;
 			if (InStatement.GetColumnValues(
@@ -492,7 +521,7 @@ bool FAssetSearchDatabase::Open(const FString& InSessionPath, const ESQLiteDatab
 		return false;
 	}
 
-	if (!ensure(Database->Execute(TEXT("CREATE VIRTUAL TABLE IF NOT EXISTS table_asset_properties_fts USING FTS5(asset_name, asset_class UNINDEXED, asset_path UNINDEXED, object_name UNINDEXED, object_path UNINDEXED, object_native_class UNINDEXED, property_name UNINDEXED, property_field UNINDEXED, property_class UNINDEXED, value_text, value_hidden, assetid UNINDEXED, content=view_asset_properties, content_rowid=rowid);"))))
+	if (!ensure(Database->Execute(TEXT("CREATE VIRTUAL TABLE IF NOT EXISTS table_asset_properties_fts USING FTS5(asset_name UNINDEXED, asset_class UNINDEXED, asset_path UNINDEXED, object_name UNINDEXED, object_path UNINDEXED, object_native_class UNINDEXED, property_name UNINDEXED, property_field UNINDEXED, property_class UNINDEXED, value_text, value_hidden, assetid UNINDEXED, content=view_asset_properties, content_rowid=rowid);"))))
 	{
 		LogLastError();
 		Close();
@@ -525,9 +554,10 @@ bool FAssetSearchDatabase::Open(const FString& InSessionPath, const ESQLiteDatab
 	}
 
 	if (!ensure(Database->Execute(
-		TEXT(" CREATE TRIGGER table_asset_properties_insert AFTER INSERT ON table_asset_properties BEGIN")
-		TEXT("     INSERT INTO table_asset_properties_fts(rowid, object_name, object_path, object_native_class, property_name, property_field, property_class, value_text, value_hidden, assetid) VALUES (new.rowid, new.object_name, new.object_path, new.object_native_class, new.property_name, new.property_field, new.property_class, new.value_text, new.value_hidden, new.assetid);")
-		TEXT(" END;")
+		TEXT("CREATE TRIGGER table_asset_properties_insert AFTER INSERT ON table_asset_properties BEGIN")
+		TEXT("	INSERT INTO table_asset_properties_fts(rowid, value_text, value_hidden) ")
+		TEXT("		VALUES (new.rowid, new.value_text, new.value_hidden);")
+		TEXT("END;")
 		)))
 	{
 		LogLastError();
@@ -544,9 +574,10 @@ bool FAssetSearchDatabase::Open(const FString& InSessionPath, const ESQLiteDatab
 	}
 
 	if (!ensure(Database->Execute(
-		TEXT(" CREATE TRIGGER table_asset_properties_delete AFTER DELETE ON table_asset_properties BEGIN")
-		TEXT("     INSERT INTO table_asset_properties_fts(table_asset_properties_fts, rowid, object_name, object_path, object_native_class, property_name, property_field, property_class, value_text, value_hidden, assetid) VALUES('delete', old.rowid, old.object_name, old.object_path, old.object_native_class, old.property_name, old.property_field, old.property_class, old.value_text, old.value_hidden, old.assetid);")
-		TEXT(" END;")
+		TEXT("CREATE TRIGGER table_asset_properties_delete AFTER DELETE ON table_asset_properties BEGIN")
+		TEXT("	INSERT INTO table_asset_properties_fts(table_asset_properties_fts, rowid, value_text, value_hidden) ")
+		TEXT("		VALUES('delete', old.rowid, old.value_text, old.value_hidden);")
+		TEXT("END;")
 	)))
 	{
 		LogLastError();
@@ -564,10 +595,12 @@ bool FAssetSearchDatabase::Open(const FString& InSessionPath, const ESQLiteDatab
 	}
 
 	if (!ensure(Database->Execute(
-		TEXT(" CREATE TRIGGER table_asset_properties_update AFTER UPDATE ON table_asset_properties BEGIN")
-		TEXT("     INSERT INTO table_asset_properties_fts(table_asset_properties_fts, rowid, object_name, object_path, object_native_class, property_name, property_field, property_class, value_text, value_hidden, assetid) VALUES('delete', old.rowid, old.object_name, old.object_path, old.object_native_class, old.property_name, old.property_field, old.property_class, old.value_text, old.value_hidden, old.assetid);")
-		TEXT("     INSERT INTO table_asset_properties_fts(rowid, object_name, object_path, object_native_class, property_name, property_field, property_class, value_text, value_hidden, assetid) VALUES (new.rowid, new.object_name, new.object_path, new.object_native_class, new.property_name, new.property_field, new.property_class, new.value_text, new.value_hidden, new.assetid);")
-		TEXT(" END;")
+		TEXT("CREATE TRIGGER table_asset_properties_update AFTER UPDATE ON table_asset_properties BEGIN")
+		TEXT("	INSERT INTO table_asset_properties_fts(table_asset_properties_fts, rowid, value_text, value_hidden) ")
+		TEXT("		VALUES('delete', old.rowid, old.value_text, old.value_hidden);")
+		TEXT("	INSERT INTO table_asset_properties_fts(rowid, value_text, value_hidden) ")
+		TEXT("		VALUES (new.rowid, new.value_text, new.value_hidden);")
+		TEXT("END;")
 	)))
 	{
 		LogLastError();
@@ -682,7 +715,16 @@ void FAssetSearchDatabase::AddOrUpdateAsset(const FAssetData& InAssetData, const
 
 bool FAssetSearchDatabase::EnumerateSearchResults(const FSearchQuery& Query, TFunctionRef<bool(FSearchRecord&&)> InCallback)
 {
-	bool bSuccess = Statements->SearchAssets(Query, [&InCallback](FSearchRecord&& InResult)
+	bool bSuccess = false;
+
+	bSuccess |= Statements->SearchAssets(Query, [&InCallback](FSearchRecord&& InResult)
+	{
+		return InCallback(MoveTemp(InResult))
+			? ESQLitePreparedStatementExecuteRowResult::Continue
+			: ESQLitePreparedStatementExecuteRowResult::Stop;
+	});
+
+	bSuccess |= Statements->SearchAssetProperties(Query, [&InCallback](FSearchRecord&& InResult)
 	{
 		return InCallback(MoveTemp(InResult))
 			? ESQLitePreparedStatementExecuteRowResult::Continue
diff --git a/Engine/Plugins/Editor/AssetSearch/Source/Private/Settings/SearchUserSettings.h b/Engine/Plugins/Editor/AssetSearch/Source/Private/Settings/SearchUserSettings.h
--- a/Engine/Plugins/Editor/AssetSearch/Source/Private/Settings/SearchUserSettings.h
+++ b/Engine/Plugins/Editor/AssetSearch/Source/Private/Settings/SearchUserSettings.h
@@ -40,6 +40,9 @@ public:
 	UPROPERTY(config, EditAnywhere, Category=General)
 	bool bShowMissingAssets = true;
 
+	UPROPERTY(config, EditAnywhere, Category = General)
+	bool bAutoExpandAssets = true;
+
 	UPROPERTY(config, EditAnywhere, Category = Performance)
 	bool bThrottleInBackground = true;
 
diff --git a/Engine/Plugins/Editor/AssetSearch/Source/Private/Utility/IndexerUtilities.cpp b/Engine/Plugins/Editor/AssetSearch/Source/Private/Utility/IndexerUtilities.cpp
--- a/Engine/Plugins/Editor/AssetSearch/Source/Private/Utility/IndexerUtilities.cpp
+++ b/Engine/Plugins/Editor/AssetSearch/Source/Private/Utility/IndexerUtilities.cpp
@@ -5,6 +5,7 @@
 #include "GameplayTagContainer.h"
 #include "Internationalization/Text.h"
 #include "UObject/TextProperty.h"
+#include "UObject/SoftObjectPtr.h"
 
 PRAGMA_DISABLE_OPTIMIZATION
 
@@ -79,6 +80,22 @@ void FIndexerUtilities::IterateIndexableProperties(const UStruct* InStruct, cons
 			FText DisplayName = EnumProperty->GetEnum()->GetDisplayNameTextByValue(Value);
 			Text = *FTextInspector::GetSourceString(DisplayName);
 		}
+		//else if (const FObjectProperty* ObjectProperty = CastField<FObjectProperty>(Property))
+		//{
+		//	UObject* Object = ObjectProperty->GetPropertyValue(ValuePtr);
+		//	if (Object && Object->HasAnyFlags(RF_Public))
+		//	{
+		//		Text = Object->GetName();
+		//	}
+		//}
+		//else if (const FSoftObjectProperty* SoftObjectProperty = CastField<FSoftObjectProperty>(Property))
+		//{
+		//	FSoftObjectPtr SoftObject = SoftObjectProperty->GetPropertyValue(ValuePtr);
+		//	if (!SoftObject.IsNull())
+		//	{
+		//		Text = SoftObject.GetAssetName();
+		//	}
+		//}
 		else if (const FStructProperty* StructProperty = CastField<FStructProperty>(Property))
 		{
 			if (StructProperty->Struct == FGameplayTag::StaticStruct())
diff --git a/Engine/Plugins/Editor/AssetSearch/Source/Private/Widgets/SSearchBrowser.cpp b/Engine/Plugins/Editor/AssetSearch/Source/Private/Widgets/SSearchBrowser.cpp
--- a/Engine/Plugins/Editor/AssetSearch/Source/Private/Widgets/SSearchBrowser.cpp
+++ b/Engine/Plugins/Editor/AssetSearch/Source/Private/Widgets/SSearchBrowser.cpp
@@ -209,13 +209,30 @@ void SSearchBrowser::Construct( const FArguments& InArgs )
 				// Index unindexed items
 				+ SHorizontalBox::Slot()
 				.AutoWidth()
-				.VAlign(VAlign_Top)
+				.VAlign(VAlign_Center)
 				[
 					SNew(SHyperlink)
 					.Text(this, &SSearchBrowser::GetUnindexedAssetsText)
 					.Visibility_Lambda([] { return GetDefault<USearchUserSettings>()->bShowMissingAssets ? EVisibility::Visible : EVisibility::Collapsed; })
 					.OnNavigate(this, &SSearchBrowser::HandleForceIndexOfAssetsMissingIndex)
 				]
+
+				// View Options Button
+				+ SHorizontalBox::Slot()
+				.AutoWidth()
+				.VAlign(VAlign_Center)
+				[
+					SNew( SComboButton )
+					.ContentPadding(0)
+					.ForegroundColor( FSlateColor::UseForeground() )
+					.ButtonStyle( FEditorStyle::Get(), "ToggleButton" )
+					.OnGetMenuContent(this, &SSearchBrowser::GetViewMenuWidget)
+					.ButtonContent()
+					[
+						SNew(SImage)
+						.Image( FEditorStyle::GetBrush("GenericViewButton") )
+					]
+				]
 			]
 		]
 	];
@@ -242,6 +259,30 @@ void SSearchBrowser::Tick(const FGeometry& AllottedGeometry, const double InCurr
 	}
 }
 
+TSharedRef<SWidget> SSearchBrowser::GetViewMenuWidget()
+{
+	FMenuBuilder MenuBuilder(true, nullptr);
+
+	MenuBuilder.BeginSection("Results", LOCTEXT("ShowHeading", "Show"));
+	{
+		MenuBuilder.AddMenuEntry(
+			LOCTEXT("ToggleAutoExpandAssets", "Auto Expand Assets"),
+			LOCTEXT("ToggleAutoExpandAssetsToolTip", "When enabled, we automatically expand the assets in the results."),
+			FSlateIcon(),
+			FUIAction(
+				FExecuteAction::CreateLambda([] { GetMutableDefault<USearchUserSettings>()->bAutoExpandAssets = !GetDefault<USearchUserSettings>()->bAutoExpandAssets; }),
+				FCanExecuteAction(),
+				FIsActionChecked::CreateLambda([] { return GetDefault<USearchUserSettings>()->bAutoExpandAssets; })
+			),
+			NAME_None,
+			EUserInterfaceActionType::ToggleButton
+		);
+	}
+	MenuBuilder.EndSection();
+
+	return MenuBuilder.MakeWidget();
+}
+
 FText SSearchBrowser::GetSearchBackgroundText() const
 {
 	if (FilterString.Len() > 0 && !IsSearching() && SearchResults.Num() == 0)
@@ -339,8 +380,10 @@ void SSearchBrowser::RefreshList()
 
 		SearchesActive++;
 
+		const bool bAutoExpandAssets = GetDefault<USearchUserSettings>()->bAutoExpandAssets;
+
 		IAssetSearchModule& SearchModule = IAssetSearchModule::Get();
-		SearchModule.Search(Query, [this](TArray<FSearchRecord>&& InResults) {
+		SearchModule.Search(Query, [this, bAutoExpandAssets](TArray<FSearchRecord>&& InResults) {
 			check(IsInGameThread());
 
 			SearchResults.Reset();
@@ -355,7 +398,10 @@ void SSearchBrowser::RefreshList()
 			{
 				SearchResults.Add(Entry.Value);
 
-				SearchTreeView->SetItemExpansion(Entry.Value, true);
+				if (bAutoExpandAssets)
+				{
+					SearchTreeView->SetItemExpansion(Entry.Value, true);
+				}
 			}
 
 			SearchResults.Sort([](const TSharedPtr<FSearchNode>& A, const TSharedPtr<FSearchNode>& B) {
diff --git a/Engine/Plugins/Editor/AssetSearch/Source/Private/Widgets/SSearchBrowser.h b/Engine/Plugins/Editor/AssetSearch/Source/Private/Widgets/SSearchBrowser.h
--- a/Engine/Plugins/Editor/AssetSearch/Source/Private/Widgets/SSearchBrowser.h
+++ b/Engine/Plugins/Editor/AssetSearch/Source/Private/Widgets/SSearchBrowser.h
@@ -44,6 +44,8 @@ public:
 
 private:
 
+	TSharedRef<SWidget> GetViewMenuWidget();
+
 	FText GetSearchBackgroundText() const;
 	FText GetStatusText() const;
 	FText GetAdvancedStatus() const;
diff --git a/Engine/Plugins/Editor/AssetSearch/Source/Private/Widgets/SearchModel.h b/Engine/Plugins/Editor/AssetSearch/Source/Private/Widgets/SearchModel.h
--- a/Engine/Plugins/Editor/AssetSearch/Source/Private/Widgets/SearchModel.h
+++ b/Engine/Plugins/Editor/AssetSearch/Source/Private/Widgets/SearchModel.h
@@ -123,22 +123,25 @@ public:
 			return;
 		}
 
-		TSharedPtr<FAssetObjectNode> ExistingObject = Objects.FindRef(InResult.object_path);
-		if (!ExistingObject.IsValid())
+		if (!InResult.object_path.IsEmpty())
 		{
-			ExistingObject = MakeShared<FAssetObjectNode>(InResult);
-			Objects.Add(InResult.object_path, ExistingObject);
-		}
-		else
-		{
-			ExistingObject->Append(InResult);
+			TSharedPtr<FAssetObjectNode> ExistingObject = Objects.FindRef(InResult.object_path);
+			if (!ExistingObject.IsValid())
+			{
+				ExistingObject = MakeShared<FAssetObjectNode>(InResult);
+				Objects.Add(InResult.object_path, ExistingObject);
+			}
+			else
+			{
+				ExistingObject->Append(InResult);
+			}
+			
+			SortedObjectArray.Reset();
+			Objects.GenerateValueArray(SortedObjectArray);
 		}
 
 		TotalScore += InResult.Score;
 		MaxScore = FMath::Min(MaxScore, InResult.Score);
-
-		SortedObjectArray.Reset();
-		Objects.GenerateValueArray(SortedObjectArray);
 	}
 
 	virtual ESearchNodeType GetType() const override { return ESearchNodeType::Asset; }
