Search - Avoid going back to the MT to enqueue new work.  Avoid going back to the MT for all key requests.

Fixes a shutdown problem where GT tasks where executing after destruction.

diff --git a/Engine/Plugins/Editor/AssetSearch/Source/Private/AssetSearchManager.cpp b/Engine/Plugins/Editor/AssetSearch/Source/Private/AssetSearchManager.cpp
--- a/Engine/Plugins/Editor/AssetSearch/Source/Private/AssetSearchManager.cpp
+++ b/Engine/Plugins/Editor/AssetSearch/Source/Private/AssetSearchManager.cpp
@@ -531,9 +531,8 @@ bool FAssetSearchManager::AsyncGetDerivedDataKey(const FAssetData& InAssetData,
 			DDCKey.Append(LexToString(FileInfo.Hash));
 
 			const FString DDCKeyString = DDCKey.ToString();
-			AsyncMainThreadTask([this, DDCKeyString, DDCKeyCallback]() {
-				DDCKeyCallback(DDCKeyString);
-			});
+
+			DDCKeyCallback(DDCKeyString);
 		}
 	});
 
@@ -604,13 +603,15 @@ void FAssetSearchManager::StoreIndexForAsset(UObject* InAsset)
 		if (bWasIndexed && !IndexedJson.IsEmpty())
 		{
 			AsyncGetDerivedDataKey(InAssetData, [this, InAssetData, IndexedJson](FString InDDCKey) {
-				check(IsInGameThread());
+				AsyncMainThreadTask([this, InAssetData, IndexedJson, InDDCKey]() {
+					check(IsInGameThread());
 
-				FTCHARToUTF8 IndexedJsonUTF8(*IndexedJson);
-				TArrayView<const uint8> IndexedJsonUTF8View((const uint8*)IndexedJsonUTF8.Get(), IndexedJsonUTF8.Length() * sizeof(UTF8CHAR));
-				GetDerivedDataCacheRef().Put(*InDDCKey, IndexedJsonUTF8View, InAssetData.ObjectPath.ToString(), false);
+					FTCHARToUTF8 IndexedJsonUTF8(*IndexedJson);
+					TArrayView<const uint8> IndexedJsonUTF8View((const uint8*)IndexedJsonUTF8.Get(), IndexedJsonUTF8.Length() * sizeof(UTF8CHAR));
+					GetDerivedDataCacheRef().Put(*InDDCKey, IndexedJsonUTF8View, InAssetData.ObjectPath.ToString(), false);
 
-				AddOrUpdateAsset(InAssetData, IndexedJson, InDDCKey);
+					AddOrUpdateAsset(InAssetData, IndexedJson, InDDCKey);
+				});
 			});
 		}
 	}
@@ -771,6 +772,8 @@ void FAssetSearchManager::ForceIndexOnAssetsMissingIndex()
 			break;
 		}
 
+		ProcessGameThreadTasks();
+
 		IndexingTask.EnterProgressFrame(1, FText::Format(LOCTEXT("ForceIndexOnAssetsMissingIndexFormat", "Indexing Asset ({0} of {1})"), RemovedCount + 1, FailedDDCRequests.Num()));
 		if (UObject* AssetToIndex = Request.AssetData.GetAsset())
 		{
@@ -846,10 +849,6 @@ void FAssetSearchManager::Search(const FSearchQuery& Query, TFunction<void(TArra
 void FAssetSearchManager::AsyncMainThreadTask(TFunction<void()> Task)
 {
 	GT_Tasks.Enqueue(Task);
-
-	AsyncTask(ENamedThreads::GameThread, [this]() {
-		ProcessGameThreadTasks();
-	});
 }
 
 void FAssetSearchManager::ProcessGameThreadTasks()
@@ -862,10 +861,13 @@ void FAssetSearchManager::ProcessGameThreadTasks()
 			return;
 		}
 
+		int MaxGameThreadTasksPerTick = 1000;
+
 		TFunction<void()> Operation;
-		if (GT_Tasks.Dequeue(Operation))
+		while (GT_Tasks.Dequeue(Operation) && MaxGameThreadTasksPerTick > 0)
 		{
 			Operation();
+			MaxGameThreadTasksPerTick--;
 		}
 	}
 }
diff --git a/Engine/Plugins/Editor/AssetSearch/Source/Private/AssetSearchManager.h b/Engine/Plugins/Editor/AssetSearch/Source/Private/AssetSearchManager.h
--- a/Engine/Plugins/Editor/AssetSearch/Source/Private/AssetSearchManager.h
+++ b/Engine/Plugins/Editor/AssetSearch/Source/Private/AssetSearchManager.h
@@ -104,7 +104,7 @@ private:
 
 	FDelegateHandle TickerHandle;
 
-	TQueue<TFunction<void()>> GT_Tasks;
+	TQueue<TFunction<void()>, EQueueMode::Mpsc> GT_Tasks;
 
 private:
 	bool bDatabaseOpen = false;
@@ -113,7 +113,7 @@ private:
 	TAtomic<bool> RunThread;
 	FRunnableThread* DatabaseThread = nullptr;
 
-	TQueue<TFunction<void()>> ImmediateOperations;
-	TQueue<TFunction<void()>> FeedOperations;
-	TQueue<TFunction<void()>> UpdateOperations;
+	TQueue<TFunction<void()>, EQueueMode::Mpsc> ImmediateOperations;
+	TQueue<TFunction<void()>, EQueueMode::Mpsc> FeedOperations;
+	TQueue<TFunction<void()>, EQueueMode::Mpsc> UpdateOperations;
 };
\ No newline at end of file
