Search - Adding some additionall logging, and better stat reporting that takes into account things being downloaded that haven't been processed yet.

diff --git a/Engine/Plugins/Editor/AssetSearch/Source/Private/AssetSearchDatabase.cpp b/Engine/Plugins/Editor/AssetSearch/Source/Private/AssetSearchDatabase.cpp
--- a/Engine/Plugins/Editor/AssetSearch/Source/Private/AssetSearchDatabase.cpp
+++ b/Engine/Plugins/Editor/AssetSearch/Source/Private/AssetSearchDatabase.cpp
@@ -171,11 +171,14 @@ public:
 	private: FAddAssetPropertiesFromJson Statement_AddAssetProperty;
 	public: bool AddSearchRecord(const FAssetData& InAssetData, const FString& IndexedJson, const FString& IndexedJsonHash)
 	{
-		if (Statement_AddAssetToAssetTable.BindAndExecute(InAssetData.AssetName.ToString(), InAssetData.AssetClass.ToString(), InAssetData.ObjectPath.ToString(), IndexedJsonHash))
+		FString AssetObjectPath = InAssetData.ObjectPath.ToString();
+		if (Statement_AddAssetToAssetTable.BindAndExecute(InAssetData.AssetName.ToString(), InAssetData.AssetClass.ToString(), AssetObjectPath, IndexedJsonHash))
 		{
 			int64 AssetId = Database.GetLastInsertRowId();
 
-			BeginTransaction();
+			int32 EntriesAdded = 0;
+
+			ensure(BeginTransaction());
 
 			TSharedRef<TJsonReader<>> JsonReader = TJsonReaderFactory<>::Create(IndexedJson);
 
@@ -259,7 +262,11 @@ public:
 
 									if (!Statement_AddAssetProperty.BindAndExecute(AssetId, object_name, object_path, object_native_class, property_name, property_field, property_class, value_text, value_hidden))
 									{
-										//Log error?
+										UE_LOG(LogAssetSearch, Error, TEXT("Failed to add entry for asset '%s': %s"), *AssetObjectPath, *Database.GetLastError());
+									}
+									else
+									{
+										EntriesAdded++;
 									}
 								}
 							}
@@ -270,7 +277,16 @@ public:
 				}
 			}
 
-			CommitTransaction();
+			if (!CommitTransaction())
+			{
+				UE_LOG(LogAssetSearch, Error, TEXT("Failed commit entries for asset '%s': %s"), *AssetObjectPath, *Database.GetLastError());
+				EntriesAdded = 0;
+			}
+
+			if (EntriesAdded > 0)
+			{
+				//UE_LOG(LogAssetSearch, Log, TEXT("Adding '%s': %s"), *AssetObjectPath, *Database.GetLastError());
+			}
 
 			return true;
 		}
diff --git a/Engine/Plugins/Editor/AssetSearch/Source/Private/AssetSearchManager.cpp b/Engine/Plugins/Editor/AssetSearch/Source/Private/AssetSearchManager.cpp
--- a/Engine/Plugins/Editor/AssetSearch/Source/Private/AssetSearchManager.cpp
+++ b/Engine/Plugins/Editor/AssetSearch/Source/Private/AssetSearchManager.cpp
@@ -49,11 +49,11 @@ FAutoConsoleVariableRef CVarDisableUniversalSearch(
 	TEXT("Enable universal search")
 );
 
-static bool bIndexUnindexAssetsOnLoad = false;
-FAutoConsoleVariableRef CVarIndexUnindexAssetsOnLoad(
-	TEXT("Search.IndexUnindexAssetsOnLoad"),
-	bIndexUnindexAssetsOnLoad,
-	TEXT("Index Unindex Assets On Load")
+static bool bTryIndexAssetsOnLoad = true;
+FAutoConsoleVariableRef CVarTryIndexAssetsOnLoad(
+	TEXT("Search.TryIndexAssetsOnLoad"),
+	bTryIndexAssetsOnLoad,
+	TEXT("Tries to index assets on load.")
 );
 
 class FUnloadPackageScope
@@ -284,7 +284,7 @@ FSearchStats FAssetSearchManager::GetStats() const
 {
 	FSearchStats Stats;
 	Stats.Scanning = ProcessAssetQueue.Num();
-	Stats.Processing = IsAssetUpToDateCount + DownloadQueueCount;
+	Stats.Processing = IsAssetUpToDateCount + DownloadQueueCount + ActiveDownloads;
 	Stats.Updating = PendingDatabaseUpdates;
 	Stats.TotalRecords = TotalSearchRecords;
 	Stats.AssetsMissingIndex = FailedDDCRequests.Num();
@@ -400,7 +400,7 @@ void FAssetSearchManager::OnAssetLoaded(UObject* InObject)
 {
 	check(IsInGameThread());
 
-	if (bIndexUnindexAssetsOnLoad)
+	if (bTryIndexAssetsOnLoad)
 	{
 		RequestIndexAsset(InObject);
 	}
@@ -410,7 +410,7 @@ bool FAssetSearchManager::RequestIndexAsset(UObject* InAsset)
 {
 	check(IsInGameThread());
 
-	if (GEditor->IsAutosaving())
+	if (GEditor == nullptr || GEditor->IsAutosaving())
 	{
 		return false;
 	}
@@ -680,13 +680,15 @@ bool FAssetSearchManager::Tick_GameThread(float DeltaTime)
 		check(bSuccess);
 
 		DownloadQueueCount--;
+		ActiveDownloads++;
 
 		DDCRequest.DDCHandle = GetDerivedDataCacheRef().GetAsynchronous(*DDCRequest.DDCKey, DDCRequest.AssetData.ObjectPath.ToString());
 		ProcessDDCQueue.Enqueue(DDCRequest);
 	}
 
+	int32 MaxQueueProcesses = 1000;
 	int32 DownloadProcessLimit = PerformanceLimits.DownloadProcessRate;
-	while (!ProcessDDCQueue.IsEmpty() && DownloadProcessLimit > 0)
+	while (!ProcessDDCQueue.IsEmpty() && DownloadProcessLimit > 0 && MaxQueueProcesses > 0)
 	{
 		const FAssetDDCRequest* PendingRequest = ProcessDDCQueue.Peek();
 		if (GetDerivedDataCacheRef().PollAsynchronousCompletion(PendingRequest->DDCHandle))
@@ -698,6 +700,7 @@ bool FAssetSearchManager::Tick_GameThread(float DeltaTime)
 			if (bGetSuccessful)
 			{
 				LoadDDCContentIntoDatabase(PendingRequest->AssetData, OutContent, PendingRequest->DDCKey);
+				DownloadProcessLimit--;
 			}
 			else if (UserSettings->bShowMissingAssets)
 			{
@@ -706,7 +709,7 @@ bool FAssetSearchManager::Tick_GameThread(float DeltaTime)
 
 			ProcessDDCQueue.Pop();
 			ActiveDownloads--;
-			DownloadProcessLimit--;
+			MaxQueueProcesses--;
 			continue;
 		}
 		break;
