Search - Disable trying indexing on load.

diff --git a/Engine/Plugins/Editor/AssetSearch/Source/Private/AssetSearchManager.cpp b/Engine/Plugins/Editor/AssetSearch/Source/Private/AssetSearchManager.cpp
--- a/Engine/Plugins/Editor/AssetSearch/Source/Private/AssetSearchManager.cpp
+++ b/Engine/Plugins/Editor/AssetSearch/Source/Private/AssetSearchManager.cpp
@@ -49,7 +49,7 @@ FAutoConsoleVariableRef CVarDisableUniversalSearch(
 	TEXT("Enable universal search")
 );
 
-static bool bTryIndexAssetsOnLoad = true;
+static bool bTryIndexAssetsOnLoad = false;
 FAutoConsoleVariableRef CVarTryIndexAssetsOnLoad(
 	TEXT("Search.TryIndexAssetsOnLoad"),
 	bTryIndexAssetsOnLoad,
@@ -760,6 +760,8 @@ void FAssetSearchManager::ForceIndexOnAssetsMissingIndex()
 {
 	check(IsInGameThread());
 
+	EAppReturnType::Type IncludeMaps = FMessageDialog::Open(EAppMsgType::YesNo, LOCTEXT("IncludeMaps", "Do you want to open and index map files, this can take a long time?"));
+
 	FScopedSlowTask IndexingTask(FailedDDCRequests.Num(), LOCTEXT("ForceIndexOnAssetsMissingIndex", "Indexing Assets"));
 	IndexingTask.MakeDialog(true);
 
@@ -775,6 +777,15 @@ void FAssetSearchManager::ForceIndexOnAssetsMissingIndex()
 			break;
 		}
 
+		if (IncludeMaps != EAppReturnType::Yes)
+		{
+			if (Request.AssetData.GetClass() == UWorld::StaticClass())
+			{
+				RemovedCount++;
+				continue;
+			}
+		}
+
 		ProcessGameThreadTasks();
 
 		IndexingTask.EnterProgressFrame(1, FText::Format(LOCTEXT("ForceIndexOnAssetsMissingIndexFormat", "Indexing Asset ({0} of {1})"), RemovedCount + 1, FailedDDCRequests.Num()));
@@ -789,7 +800,10 @@ void FAssetSearchManager::ForceIndexOnAssetsMissingIndex()
 				continue;
 			}
 
-			StoreIndexForAsset(AssetToIndex);
+			if (!bTryIndexAssetsOnLoad)
+			{
+				StoreIndexForAsset(AssetToIndex);
+			}
 		}
 
 		if (UnloadScope.GetObjectsLoaded() > 2000)
@@ -802,10 +816,10 @@ void FAssetSearchManager::ForceIndexOnAssetsMissingIndex()
 
 	if (RedirectorsWithBrokenMetadata.Num() > 0)
 	{
-		EAppReturnType::Type ReturnValue = FMessageDialog::Open(EAppMsgType::YesNo,
+		EAppReturnType::Type ResaveRedirectors = FMessageDialog::Open(EAppMsgType::YesNo,
 			LOCTEXT("ResaveRedirectors", "We found some redirectors that didn't have the correct asset metadata identifying them as redirectors.  Would you like to resave them, so that they stop appearing as missing asset indexes?"));
 
-		if (ReturnValue == EAppReturnType::Yes)
+		if (ResaveRedirectors == EAppReturnType::Yes)
 		{
 			TArray<UPackage*> PackagesToSave;
 			for (const FAssetData& BrokenAsset : RedirectorsWithBrokenMetadata)
