Fixing a loc duplicate key bug

diff --git a/Engine/Plugins/Editor/AssetSearch/Source/Private/Widgets/SSearchBrowser.cpp b/Engine/Plugins/Editor/AssetSearch/Source/Private/Widgets/SSearchBrowser.cpp
--- a/Engine/Plugins/Editor/AssetSearch/Source/Private/Widgets/SSearchBrowser.cpp
+++ b/Engine/Plugins/Editor/AssetSearch/Source/Private/Widgets/SSearchBrowser.cpp
@@ -287,7 +287,7 @@ FText SSearchBrowser::GetSearchBackgroundText() const
 {
 	if (FilterString.Len() > 0 && !IsSearching() && SearchResults.Num() == 0)
 	{
-		return LOCTEXT("SearchAllTheThings", "¯\\_(ツ)_/¯");
+		return LOCTEXT("FoundNoResults", "¯\\_(ツ)_/¯");
 	}
 	else if (FilterString.Len() == 0)
 	{
