Const correct items for the search browser.

diff --git a/Engine/Plugins/Editor/AssetSearch/Source/Private/Widgets/SSearchBrowser.cpp b/Engine/Plugins/Editor/AssetSearch/Source/Private/Widgets/SSearchBrowser.cpp
--- a/Engine/Plugins/Editor/AssetSearch/Source/Private/Widgets/SSearchBrowser.cpp
+++ b/Engine/Plugins/Editor/AssetSearch/Source/Private/Widgets/SSearchBrowser.cpp
@@ -301,9 +301,9 @@ FText SSearchBrowser::GetSearchBackgroundText() const
 
 FText SSearchBrowser::GetStatusText() const
 {
-	IAssetSearchModule& SearchModule = IAssetSearchModule::Get();
-	FSearchStats SearchStats = SearchModule.GetStats();
-	int32 UpdatingCount = SearchStats.Scanning + SearchStats.Processing + SearchStats.Updating;
+	const IAssetSearchModule& SearchModule = IAssetSearchModule::Get();
+	const FSearchStats SearchStats = SearchModule.GetStats();
+	const int32 UpdatingCount = SearchStats.Scanning + SearchStats.Processing + SearchStats.Updating;
 
 	if (UpdatingCount > 0)
 	{
