Slightly improving the blueprint indexing for search, inclusing Event Member reference, so that the non-display name is included.  Not bumping the version for this small improvement.

diff --git a/Engine/Plugins/Editor/AssetSearch/Source/Private/Indexers/BlueprintIndexer.cpp b/Engine/Plugins/Editor/AssetSearch/Source/Private/Indexers/BlueprintIndexer.cpp
--- a/Engine/Plugins/Editor/AssetSearch/Source/Private/Indexers/BlueprintIndexer.cpp
+++ b/Engine/Plugins/Editor/AssetSearch/Source/Private/Indexers/BlueprintIndexer.cpp
@@ -13,6 +13,7 @@
 #include "EdGraphNode_Comment.h"
 #include "Engine/SimpleConstructionScript.h"
 #include "Engine/SCS_Node.h"
+#include "K2Node_Event.h"
 
 #define LOCTEXT_NAMESPACE "FBlueprintIndexer"
 
@@ -109,7 +110,7 @@ void FBlueprintIndexer::IndexGraphs(const UBlueprint* InBlueprint, FSearchSerial
 
 			const FText NodeText = Node->GetNodeTitle(ENodeTitleType::MenuTitle);
 			Serializer.BeginIndexingObject(Node, NodeText);
-			Serializer.IndexProperty(TEXT("Name"), NodeText);
+			Serializer.IndexProperty(TEXT("Title"), NodeText);
 
 			if (!Node->NodeComment.IsEmpty())
 			{
@@ -120,6 +121,10 @@ void FBlueprintIndexer::IndexGraphs(const UBlueprint* InBlueprint, FSearchSerial
 			{
 				IndexMemberReference(Serializer, FunctionNode->FunctionReference, TEXT("Function"));
 			}
+			else if (UK2Node_Event* EventNode = Cast<UK2Node_Event>(Node))
+			{
+				IndexMemberReference(Serializer, EventNode->EventReference, TEXT("Event"));
+			}
 			else if (UK2Node_BaseMCDelegate* DelegateNode = Cast<UK2Node_BaseMCDelegate>(Node))
 			{
 				IndexMemberReference(Serializer, DelegateNode->DelegateReference, TEXT("Delegate"));
