Fixing loctext collision for asset search

diff --git a/Engine/Plugins/Editor/AssetSearch/Source/Private/AssetSearchModule.cpp b/Engine/Plugins/Editor/AssetSearch/Source/Private/AssetSearchModule.cpp
--- a/Engine/Plugins/Editor/AssetSearch/Source/Private/AssetSearchModule.cpp
+++ b/Engine/Plugins/Editor/AssetSearch/Source/Private/AssetSearchModule.cpp
@@ -34,7 +34,7 @@ public:
 
 		FGlobalTabmanager::Get()->RegisterNomadTabSpawner(SearchTabName, FOnSpawnTab::CreateRaw(this, &FAssetSearchModule::HandleSpawnSettingsTab))
 			.SetDisplayName(LOCTEXT("Search", "Search"))
-			.SetTooltipText(LOCTEXT("Search", "Search Tab"))
+			.SetTooltipText(LOCTEXT("SearchTab", "Search Tab"))
 			.SetIcon(FSlateIcon(FEditorStyle::GetStyleSetName(), "Symbols.SearchGlass"))
 			.SetGroup(WorkspaceMenu::GetMenuStructure().GetToolsCategory());
 	}
