Universal search - Adding support for curve tables, updating the blueprint version to capture some of the improvements.

diff --git a/Engine/Plugins/Editor/AssetSearch/Source/Private/AssetSearchManager.cpp b/Engine/Plugins/Editor/AssetSearch/Source/Private/AssetSearchManager.cpp
--- a/Engine/Plugins/Editor/AssetSearch/Source/Private/AssetSearchManager.cpp
+++ b/Engine/Plugins/Editor/AssetSearch/Source/Private/AssetSearchManager.cpp
@@ -7,10 +7,6 @@
 #include "Async/Async.h"
 #include "DerivedDataCacheInterface.h"
 
-#include "Indexers/DataAssetIndexer.h"
-#include "Indexers/DataTableIndexer.h"
-#include "Indexers/BlueprintIndexer.h"
-#include "Indexers/WidgetBlueprintIndexer.h"
 #include "Containers/StringConv.h"
 #include "Containers/Ticker.h"
 #include "Misc/Paths.h"
@@ -22,12 +18,9 @@
 #include "Engine/Blueprint.h"
 #include "Engine/DataTable.h"
 #include "Engine/DataAsset.h"
-#include "Indexers/DialogueWaveIndexer.h"
 #include "Sound/DialogueWave.h"
-#include "Indexers/LevelIndexer.h"
 #include "Settings/SearchProjectSettings.h"
 #include "Settings/SearchUserSettings.h"
-#include "Indexers/SoundCueIndexer.h"
 #include "Sound/SoundCue.h"
 #include "Misc/ScopedSlowTask.h"
 #include "Misc/StringBuilder.h"
@@ -39,6 +32,16 @@
 #include "UObject/WeakObjectPtrTemplates.h"
 #include "FileHelpers.h"
 #include "Misc/MessageDialog.h"
+#include "Engine/CurveTable.h"
+
+#include "Indexers/DataAssetIndexer.h"
+#include "Indexers/DataTableIndexer.h"
+#include "Indexers/BlueprintIndexer.h"
+#include "Indexers/WidgetBlueprintIndexer.h"
+#include "Indexers/CurveTableIndexer.h"
+#include "Indexers/DialogueWaveIndexer.h"
+#include "Indexers/LevelIndexer.h"
+#include "Indexers/SoundCueIndexer.h"
 
 #define LOCTEXT_NAMESPACE "FAssetSearchManager"
 
@@ -56,22 +59,40 @@ FAutoConsoleVariableRef CVarTryIndexAssetsOnLoad(
 	TEXT("Tries to index assets on load.")
 );
 
+static bool bTryToGCDuringMissingIndexing = false;
+FAutoConsoleVariableRef CVarTryToGCDuringMissingIndexing(
+	TEXT("Search.TryToGCDuringMissingIndexing"),
+	bTryToGCDuringMissingIndexing,
+	TEXT("Tries to GC occasionally while indexing missing items.")
+);
+
 class FUnloadPackageScope
 {
 public:
 	FUnloadPackageScope()
 	{
-		FCoreUObjectDelegates::OnAssetLoaded.AddRaw(this, &FUnloadPackageScope::OnAssetLoaded);
+		if (bTryToGCDuringMissingIndexing)
+		{
+			FCoreUObjectDelegates::OnAssetLoaded.AddRaw(this, &FUnloadPackageScope::OnAssetLoaded);
+		}
 	}
 
 	~FUnloadPackageScope()
 	{
-		FCoreUObjectDelegates::OnAssetLoaded.RemoveAll(this);
-		TryUnload(true);
+		if (bTryToGCDuringMissingIndexing)
+		{
+			FCoreUObjectDelegates::OnAssetLoaded.RemoveAll(this);
+			TryUnload(true);
+		}
 	}
 
 	int32 TryUnload(bool bResetTrackedObjects)
 	{
+		if (!bTryToGCDuringMissingIndexing)
+		{
+			return 0;
+		}
+
 		TArray<TWeakObjectPtr<UObject>> PackageObjectPtrs;
 
 		for (const FObjectKey& LoadedObjectKey : ObjectsLoaded)
@@ -181,6 +202,7 @@ void FAssetSearchManager::Start()
 {
 	RegisterAssetIndexer(UDataAsset::StaticClass(), MakeUnique<FDataAssetIndexer>());
 	RegisterAssetIndexer(UDataTable::StaticClass(), MakeUnique<FDataTableIndexer>());
+	RegisterAssetIndexer(UCurveTable::StaticClass(), MakeUnique<FCurveTableIndexer>());
 	RegisterAssetIndexer(UBlueprint::StaticClass(), MakeUnique<FBlueprintIndexer>());
 	RegisterAssetIndexer(UWidgetBlueprint::StaticClass(), MakeUnique<FWidgetBlueprintIndexer>());
 	RegisterAssetIndexer(UDialogueWave::StaticClass(), MakeUnique<FDialogueWaveIndexer>());
@@ -777,6 +799,8 @@ void FAssetSearchManager::ForceIndexOnAssetsMissingIndex()
 			break;
 		}
 
+		IndexingTask.EnterProgressFrame(1, FText::Format(LOCTEXT("ForceIndexOnAssetsMissingIndexFormat", "Indexing Asset ({0} of {1})"), RemovedCount + 1, FailedDDCRequests.Num()));
+
 		if (IncludeMaps != EAppReturnType::Yes)
 		{
 			if (Request.AssetData.GetClass() == UWorld::StaticClass())
@@ -788,7 +812,6 @@ void FAssetSearchManager::ForceIndexOnAssetsMissingIndex()
 
 		ProcessGameThreadTasks();
 
-		IndexingTask.EnterProgressFrame(1, FText::Format(LOCTEXT("ForceIndexOnAssetsMissingIndexFormat", "Indexing Asset ({0} of {1})"), RemovedCount + 1, FailedDDCRequests.Num()));
 		if (UObject* AssetToIndex = Request.AssetData.GetAsset())
 		{
 			// This object's metadata incorrectly labled it as something other than a redirector.  We need to resave it
diff --git a/Engine/Plugins/Editor/AssetSearch/Source/Private/Indexers/BlueprintIndexer.cpp b/Engine/Plugins/Editor/AssetSearch/Source/Private/Indexers/BlueprintIndexer.cpp
--- a/Engine/Plugins/Editor/AssetSearch/Source/Private/Indexers/BlueprintIndexer.cpp
+++ b/Engine/Plugins/Editor/AssetSearch/Source/Private/Indexers/BlueprintIndexer.cpp
@@ -26,6 +26,7 @@ enum class EBlueprintIndexerVersion
 	FixingPinsToSaveValues,
 	IndexingPublicEditableFieldsOnNodes,
 	DontIndexPinsUnlessItsInputWithNoConnections,
+	BetterSupportForIndexingEventNodes,
 
 	// -----<new versions can be added above this line>-------------------------------------------------
 	VersionPlusOne,
diff --git a/Engine/Plugins/Editor/AssetSearch/Source/Private/Indexers/CurveTableIndexer.cpp b/Engine/Plugins/Editor/AssetSearch/Source/Private/Indexers/CurveTableIndexer.cpp
new file mode 100644
--- /dev/null
+++ b/Engine/Plugins/Editor/AssetSearch/Source/Private/Indexers/CurveTableIndexer.cpp
@@ -0,0 +1,36 @@
+// Copyright Epic Games, Inc. All Rights Reserved.
+
+#include "CurveTableIndexer.h"
+#include "Utility/IndexerUtilities.h"
+#include "Engine/CurveTable.h"
+
+enum class ECurveTableIndexerVersion
+{
+	Empty,
+	Initial,
+
+	// -----<new versions can be added above this line>-------------------------------------------------
+	VersionPlusOne,
+	LatestVersion = VersionPlusOne - 1
+};
+
+int32 FCurveTableIndexer::GetVersion() const
+{
+	return (int32)ECurveTableIndexerVersion::LatestVersion;
+}
+
+void FCurveTableIndexer::IndexAsset(const UObject* InAssetObject, FSearchSerializer& Serializer) const
+{
+	const UCurveTable* DataTable = Cast<UCurveTable>(InAssetObject);
+	check(DataTable);
+
+	const TMap<FName, FRealCurve*>& Rows = DataTable->GetRowMap();
+
+	Serializer.BeginIndexingObject(DataTable, TEXT("$self"));
+	for (const auto& Entry : Rows)
+	{
+		const FName& RowName = Entry.Key;
+		Serializer.IndexProperty(TEXT("Row"), RowName);
+	}
+	Serializer.EndIndexingObject();
+}
\ No newline at end of file
diff --git a/Engine/Plugins/Editor/AssetSearch/Source/Private/Indexers/CurveTableIndexer.h b/Engine/Plugins/Editor/AssetSearch/Source/Private/Indexers/CurveTableIndexer.h
new file mode 100644
--- /dev/null
+++ b/Engine/Plugins/Editor/AssetSearch/Source/Private/Indexers/CurveTableIndexer.h
@@ -0,0 +1,13 @@
+// Copyright Epic Games, Inc. All Rights Reserved.
+
+#pragma once
+
+#include "CoreMinimal.h"
+#include "IAssetSearchModule.h"
+
+class FCurveTableIndexer : public IAssetIndexer
+{
+	virtual FString GetName() const override { return TEXT("CurveTable"); }
+	virtual int32 GetVersion() const override;
+	virtual void IndexAsset(const UObject* InAssetObject, FSearchSerializer& Serializer) const override;
+};
\ No newline at end of file
diff --git a/Engine/Plugins/Editor/AssetSearch/Source/Private/Indexers/DataTableIndexer.cpp b/Engine/Plugins/Editor/AssetSearch/Source/Private/Indexers/DataTableIndexer.cpp
--- a/Engine/Plugins/Editor/AssetSearch/Source/Private/Indexers/DataTableIndexer.cpp
+++ b/Engine/Plugins/Editor/AssetSearch/Source/Private/Indexers/DataTableIndexer.cpp
@@ -8,6 +8,7 @@ enum class EDataTableIndexerVersion
 {
 	Empty,
 	Initial,
+	AddedIndexingDataTableRowNames,
 
 	// -----<new versions can be added above this line>-------------------------------------------------
 	VersionPlusOne,
@@ -32,6 +33,8 @@ void FDataTableIndexer::IndexAsset(const UObject* InAssetObject, FSearchSerializ
 		const FName& RowName = Entry.Key;
 		uint8* Row = Entry.Value;
 
+		Serializer.IndexProperty(TEXT("Row"), RowName);
+
 		FIndexerUtilities::IterateIndexableProperties(DataTable->GetRowStruct(), Row, [&Serializer](const FProperty* Property, const FString& Value) {
 			Serializer.IndexProperty(Property, Value);
 		});
