Universal Search - Tweaking some display elements.

diff --git a/Engine/Plugins/Editor/AssetSearch/Source/Private/Widgets/SSearchBrowser.cpp b/Engine/Plugins/Editor/AssetSearch/Source/Private/Widgets/SSearchBrowser.cpp
--- a/Engine/Plugins/Editor/AssetSearch/Source/Private/Widgets/SSearchBrowser.cpp
+++ b/Engine/Plugins/Editor/AssetSearch/Source/Private/Widgets/SSearchBrowser.cpp
@@ -291,7 +291,17 @@ FText SSearchBrowser::GetSearchBackgroundText() const
 	}
 	else if (FilterString.Len() == 0)
 	{
-		return LOCTEXT("SearchAllTheThings", "Search\nAll The Things!");
+		IAssetSearchModule& SearchModule = IAssetSearchModule::Get();
+		FSearchStats SearchStats = SearchModule.GetStats();
+
+		if (SearchStats.IsUpdating())
+		{
+			return FText::Format(LOCTEXT("SearchNumberOfThings", "Search\n{0} Things!"), SearchStats.TotalRecords);
+		}
+		else
+		{
+			return LOCTEXT("SearchAllTheThings", "Search\nAll The Things!");
+		}
 	}
 	else
 	{
@@ -303,9 +313,8 @@ FText SSearchBrowser::GetStatusText() const
 {
 	const IAssetSearchModule& SearchModule = IAssetSearchModule::Get();
 	const FSearchStats SearchStats = SearchModule.GetStats();
-	const int32 UpdatingCount = SearchStats.Scanning + SearchStats.Processing + SearchStats.Updating;
-
-	if (UpdatingCount > 0)
+	
+	if (SearchStats.IsUpdating())
 	{
 		return LOCTEXT("Updating", "Updating...  (You can search any time)");
 	}
diff --git a/Engine/Plugins/Editor/AssetSearch/Source/Public/IAssetSearchModule.h b/Engine/Plugins/Editor/AssetSearch/Source/Public/IAssetSearchModule.h
--- a/Engine/Plugins/Editor/AssetSearch/Source/Public/IAssetSearchModule.h
+++ b/Engine/Plugins/Editor/AssetSearch/Source/Public/IAssetSearchModule.h
@@ -47,6 +47,11 @@ struct FSearchStats
 	int32 AssetsMissingIndex = 0;
 
 	int64 TotalRecords = 0;
+
+	bool IsUpdating() const
+	{
+		return (Scanning + Processing + Updating) > 0;
+	}
 };
 
 class IAssetIndexer
