Lyra - Fixing some issues with super search, it didn't correctly handle external actors and ignoring their asset files.

diff --git a/Engine/Plugins/Editor/AssetSearch/Source/Private/AssetSearchManager.cpp b/Engine/Plugins/Editor/AssetSearch/Source/Private/AssetSearchManager.cpp
--- a/Engine/Plugins/Editor/AssetSearch/Source/Private/AssetSearchManager.cpp
+++ b/Engine/Plugins/Editor/AssetSearch/Source/Private/AssetSearchManager.cpp
@@ -550,7 +550,7 @@ bool FAssetSearchManager::TryLoadIndexForAsset(const FAssetData& InAssetData)
 			FScopeLock ScopedLock(&SearchDatabaseCS);
 			if (!SearchDatabase.IsAssetUpToDate(InAssetData, InDDCKey))
 			{
-				AsyncRequestDownlaod(InAssetData, InDDCKey);
+				AsyncRequestDownload(InAssetData, InDDCKey);
 			}
 
 			IsAssetUpToDateCount--;
@@ -565,7 +565,7 @@ bool FAssetSearchManager::TryLoadIndexForAsset(const FAssetData& InAssetData)
 	return bSuccess;
 }
 
-void FAssetSearchManager::AsyncRequestDownlaod(const FAssetData& InAssetData, const FString& InDDCKey)
+void FAssetSearchManager::AsyncRequestDownload(const FAssetData& InAssetData, const FString& InDDCKey)
 {
 	DownloadQueueCount++;
 
@@ -579,6 +579,12 @@ bool FAssetSearchManager::AsyncGetDerivedDataKey(const FAssetData& InAssetData,
 {
 	check(IsInGameThread());
 
+	const FString AssetPath = InAssetData.PackagePath.ToString();
+	if (AssetPath.Contains(FPackagePath::GetExternalActorsFolderName()) || AssetPath.Contains(FPackagePath::GetExternalObjectsFolderName()))
+	{
+		return false;
+	}
+
 	FString IndexersNamesAndVersions = GetIndexerVersion(InAssetData.GetClass());
 
 	// If the indexer names and versions is empty, then we know it's not possible to index this type of thing.
diff --git a/Engine/Plugins/Editor/AssetSearch/Source/Private/AssetSearchManager.h b/Engine/Plugins/Editor/AssetSearch/Source/Private/AssetSearchManager.h
--- a/Engine/Plugins/Editor/AssetSearch/Source/Private/AssetSearchManager.h
+++ b/Engine/Plugins/Editor/AssetSearch/Source/Private/AssetSearchManager.h
@@ -56,7 +56,7 @@ private:
 	bool RequestIndexAsset(UObject* InAsset);
 	bool IsAssetIndexable(UObject* InAsset);
 	bool TryLoadIndexForAsset(const FAssetData& InAsset);
-	void AsyncRequestDownlaod(const FAssetData& InAssetData, const FString& InDDCKey);
+	void AsyncRequestDownload(const FAssetData& InAssetData, const FString& InDDCKey);
 	bool AsyncGetDerivedDataKey(const FAssetData& UnindexedAsset, TFunction<void(bool, FString)> DDCKeyCallback);
 	bool HasIndexerForClass(const UClass* InAssetClass) const;
 	FString GetIndexerVersion(const UClass* InAssetClass) const;
diff --git a/Engine/Plugins/Editor/AssetSearch/Source/Private/Widgets/SSearchBrowser.cpp b/Engine/Plugins/Editor/AssetSearch/Source/Private/Widgets/SSearchBrowser.cpp
--- a/Engine/Plugins/Editor/AssetSearch/Source/Private/Widgets/SSearchBrowser.cpp
+++ b/Engine/Plugins/Editor/AssetSearch/Source/Private/Widgets/SSearchBrowser.cpp
@@ -118,20 +118,6 @@ void SSearchBrowser::Construct( const FArguments& InArgs )
 				[
 					SNew(SOverlay)
 
-					+ SOverlay::Slot()
-					.HAlign(HAlign_Center)
-					.VAlign(VAlign_Center)
-					[
-						SNew(STextBlock)
-						.Clipping(EWidgetClipping::Inherit)
-						.Text(this, &SSearchBrowser::GetSearchBackgroundText)
-						.Justification(ETextJustify::Center)
-						.Font(FCoreStyle::GetDefaultFontStyle("Regular", 30))
-						.ColorAndOpacity(FLinearColor(1,1,1,0.05))
-						.RenderTransformPivot(FVector2D(0.5, 0.5))
-						.RenderTransform(FSlateRenderTransform(FQuat2D(FMath::DegreesToRadians(-30.0f))))
-					]
-
 					+ SOverlay::Slot()
 					[
 						SAssignNew(SearchTreeView, STreeView< TSharedPtr<FSearchNode> >)
@@ -150,16 +136,24 @@ void SSearchBrowser::Construct( const FArguments& InArgs )
 							.FillWidth(70)
 						)
 					]
-				]
-			]
 
-			+ SVerticalBox::Slot()
-			.AutoHeight()
-			.Padding(0, 0, 0, 1)
-			[
-				SNew(SSeparator)
+					+ SOverlay::Slot()
+					.HAlign(HAlign_Center)
+					.VAlign(VAlign_Center)
+					[
+						SNew(STextBlock)
+						.Visibility(EVisibility::HitTestInvisible)
+						.Clipping(EWidgetClipping::Inherit)
+						.Text(this, &SSearchBrowser::GetSearchBackgroundText)
+						.Justification(ETextJustify::Center)
+						.Font(FCoreStyle::GetDefaultFontStyle("Regular", 30))
+						.ColorAndOpacity(FLinearColor(1,1,1,0.05))
+						.RenderTransformPivot(FVector2D(0.5, 0.5))
+						.RenderTransform(FSlateRenderTransform(FQuat2D(FMath::DegreesToRadians(-30.0f))))
+					]
+				]
 			]
-
+			
 			+ SVerticalBox::Slot()
 			.AutoHeight()
 			.Padding(0, 0, 0, 1)
@@ -283,8 +277,8 @@ FText SSearchBrowser::GetSearchBackgroundText() const
 	}
 	else if (FilterString.Len() == 0)
 	{
-		IAssetSearchModule& SearchModule = IAssetSearchModule::Get();
-		FSearchStats SearchStats = SearchModule.GetStats();
+		const IAssetSearchModule& SearchModule = IAssetSearchModule::Get();
+		const FSearchStats SearchStats = SearchModule.GetStats();
 
 		if (SearchStats.IsUpdating() && SearchStats.TotalRecords > 0)
 		{
@@ -318,15 +312,15 @@ FText SSearchBrowser::GetStatusText() const
 
 FText SSearchBrowser::GetAdvancedStatus() const
 {
-	IAssetSearchModule& SearchModule = IAssetSearchModule::Get();
-	FSearchStats SearchStats = SearchModule.GetStats();
+	const IAssetSearchModule& SearchModule = IAssetSearchModule::Get();
+	const FSearchStats SearchStats = SearchModule.GetStats();
 	return FText::Format(LOCTEXT("AdvancedSearchStatusTextFmt", "Scanning {0}\nProcessing {1}\nUpdating {2}\n\nTotal Records {3}"), SearchStats.Scanning, SearchStats.Processing, SearchStats.Updating, SearchStats.TotalRecords);
 }
 
 FText SSearchBrowser::GetUnindexedAssetsText() const
 {
-	IAssetSearchModule& SearchModule = IAssetSearchModule::Get();
-	FSearchStats SearchStats = SearchModule.GetStats();
+	const IAssetSearchModule& SearchModule = IAssetSearchModule::Get();
+	const FSearchStats SearchStats = SearchModule.GetStats();
 
 	if (SearchStats.AssetsMissingIndex > 0)
 	{
