Super Search - Adding initial integrity checking right after opening the database, and changing the sychronization to full for the file and search databases, they were getting corrupted too often under "normal" mode.

diff --git a/Engine/Plugins/Editor/AssetSearch/Source/Private/AssetSearchDatabase.cpp b/Engine/Plugins/Editor/AssetSearch/Source/Private/AssetSearchDatabase.cpp
--- a/Engine/Plugins/Editor/AssetSearch/Source/Private/AssetSearchDatabase.cpp
+++ b/Engine/Plugins/Editor/AssetSearch/Source/Private/AssetSearchDatabase.cpp
@@ -490,18 +490,24 @@ bool FAssetSearchDatabase::Open(const FString& InSessionPath, const ESQLiteDatab
 		return false;
 	}
 
+	if (!Database->PerformQuickIntegrityCheck())
+	{
+		UE_LOG(LogAssetSearch, Error, TEXT("Database failed integrity check, deleting."));
+
+		const bool bDeleteTheDatabase = true;
+		Close(bDeleteTheDatabase);
+
+		return false;
+	}
+
 	// Set the database to use exclusive WAL mode for performance (exclusive works even on platforms without a mmap implementation)
 	// Set the database "NORMAL" fsync mode to only perform a fsync when check-pointing the WAL to the main database file (fewer fsync calls are better for performance, with a very slight loss of WAL durability if the power fails)
+	Database->Execute(TEXT("PRAGMA journal_mode=WAL;"));
+	Database->Execute(TEXT("PRAGMA synchronous=FULL;"));
 	Database->Execute(TEXT("PRAGMA cache_size=1000;"));
 	Database->Execute(TEXT("PRAGMA page_size=65535;"));
 	Database->Execute(TEXT("PRAGMA locking_mode=EXCLUSIVE;"));
 
-	Database->Execute(TEXT("PRAGMA journal_mode=WAL;"));
-	Database->Execute(TEXT("PRAGMA synchronous=NORMAL;"));
-
-	/*Database->Execute(TEXT("PRAGMA journal_mode=NORMAL;"));
-	Database->Execute(TEXT("PRAGMA synchronous=OFF;"));*/
-
 	int32 LoadedDatabaseVersion = 0;
 	Database->GetUserVersion(LoadedDatabaseVersion);
 	if (LoadedDatabaseVersion != (int32)EAssetSearchDatabaseVersion::Empty)
diff --git a/Engine/Plugins/Editor/AssetSearch/Source/Private/FileInfoDatabase.cpp b/Engine/Plugins/Editor/AssetSearch/Source/Private/FileInfoDatabase.cpp
--- a/Engine/Plugins/Editor/AssetSearch/Source/Private/FileInfoDatabase.cpp
+++ b/Engine/Plugins/Editor/AssetSearch/Source/Private/FileInfoDatabase.cpp
@@ -276,6 +276,8 @@ bool FFileInfoDatabase::Open(const FString& InSessionPath)
 
 bool FFileInfoDatabase::Open(const FString& InSessionPath, const ESQLiteDatabaseOpenMode InOpenMode)
 {
+	SessionPath = InSessionPath;
+	
 	if (Database->IsValid())
 	{
 		return false;
@@ -287,20 +289,24 @@ bool FFileInfoDatabase::Open(const FString& InSessionPath, const ESQLiteDatabase
 		return false;
 	}
 
-	SessionPath = InSessionPath;
+	if (!Database->PerformQuickIntegrityCheck())
+	{
+		UE_LOG(LogFileInfo, Error, TEXT("Database failed integrity check, deleting."));
+
+		const bool bDeleteTheDatabase = true;
+		Close(bDeleteTheDatabase);
+
+		return false;
+	}
 
 	// Set the database to use exclusive WAL mode for performance (exclusive works even on platforms without a mmap implementation)
 	// Set the database "NORMAL" fsync mode to only perform a fsync when check-pointing the WAL to the main database file (fewer fsync calls are better for performance, with a very slight loss of WAL durability if the power fails)
+	Database->Execute(TEXT("PRAGMA journal_mode=WAL;"));
+	Database->Execute(TEXT("PRAGMA synchronous=FULL;"));
 	Database->Execute(TEXT("PRAGMA cache_size=1000;"));
 	Database->Execute(TEXT("PRAGMA page_size=65535;"));
 	Database->Execute(TEXT("PRAGMA locking_mode=EXCLUSIVE;"));
 
-	Database->Execute(TEXT("PRAGMA journal_mode=WAL;"));
-	Database->Execute(TEXT("PRAGMA synchronous=NORMAL;"));
-
-	/*Database->Execute(TEXT("PRAGMA journal_mode=NORMAL;"));
-	Database->Execute(TEXT("PRAGMA synchronous=OFF;"));*/
-
 	int32 LoadedDatabaseVersion = 0;
 	Database->GetUserVersion(LoadedDatabaseVersion);
 	if (LoadedDatabaseVersion != (int32)EFileInfoDatabaseVersion::Empty)
@@ -352,6 +358,16 @@ bool FFileInfoDatabase::Open(const FString& InSessionPath, const ESQLiteDatabase
 		return false;
 	}
 
+	if (!Database->PerformQuickIntegrityCheck())
+	{
+		UE_LOG(LogFileInfo, Error, TEXT("Database failed integrity check, deleting."));
+
+		const bool bDeleteTheDatabase = true;
+		Close(bDeleteTheDatabase);
+
+		return false;
+	}
+
 	return true;
 }
 
