SuperSearch - When iterating over properties to serialize for searching, we were ignoring EditInlineNew UObjects inside structs and arrays, leaving out precious search data.  That is no longer the case, we now include those items in what we search.

diff --git a/Engine/Plugins/Editor/AssetSearch/Source/Private/SearchSerializer.cpp b/Engine/Plugins/Editor/AssetSearch/Source/Private/SearchSerializer.cpp
--- a/Engine/Plugins/Editor/AssetSearch/Source/Private/SearchSerializer.cpp
+++ b/Engine/Plugins/Editor/AssetSearch/Source/Private/SearchSerializer.cpp
@@ -15,7 +15,8 @@ enum class ESearchIndexVersion
 {
 	Empty,
 	Initial,
-
+	IndexingInstancedSubObjects,
+	
 	// -----<new versions can be added above this line>-------------------------------------------------
 	VersionPlusOne,
 	LatestVersion = VersionPlusOne - 1
diff --git a/Engine/Plugins/Editor/AssetSearch/Source/Private/Utility/IndexerUtilities.cpp b/Engine/Plugins/Editor/AssetSearch/Source/Private/Utility/IndexerUtilities.cpp
--- a/Engine/Plugins/Editor/AssetSearch/Source/Private/Utility/IndexerUtilities.cpp
+++ b/Engine/Plugins/Editor/AssetSearch/Source/Private/Utility/IndexerUtilities.cpp
@@ -84,10 +84,27 @@ void FIndexerUtilities::IterateIndexableProperties(const UStruct* InStruct, cons
 		}
 		else if (const FObjectProperty* ObjectProperty = CastField<FObjectProperty>(Property))
 		{
-			UObject* Object = ObjectProperty->GetPropertyValue(ValuePtr);
-			if (Object && Object->HasAnyFlags(RF_Public))
+			if (const UObject* Object = ObjectProperty->GetPropertyValue(ValuePtr))
 			{
-				Text = Object->GetName();
+				if (Object->HasAnyFlags(RF_Public) && Object->IsAsset())
+				{
+					Text = Object->GetName();
+				}
+				else if (Object->HasAnyFlags(RF_Transient))
+				{
+					// Don't do anything with transient objects.
+					continue;
+				}
+				// If the property is "Instanced" then we may need to iterate through the object, only do this for
+				// edit inline new classes, we don't care about inner-reference objects that already already tracked
+				// in some other way, that's up to the caller of this function to handle, we just want to handle
+				// iterating the obviously owned data.
+				else if (Property->HasAllPropertyFlags(CPF_ExportObject) && Object->GetClass()->HasAllClassFlags(CLASS_EditInlineNew))
+				{
+					// Add the inner properties of this instanced object.
+					IterateIndexableProperties(Object, Callback);
+					continue;
+				}
 			}
 		}
 		else if (const FSoftObjectProperty* SoftObjectProperty = CastField<FSoftObjectProperty>(Property))
