PR #9070: Added ASSETSEARCH_API to Asset Indexer classes (Contributed by [REDACTED_PERSON])

diff --git a/Engine/Plugins/Editor/AssetSearch/Source/Public/IAssetIndexer.h b/Engine/Plugins/Editor/AssetSearch/Source/Public/IAssetIndexer.h
--- a/Engine/Plugins/Editor/AssetSearch/Source/Public/IAssetIndexer.h
+++ b/Engine/Plugins/Editor/AssetSearch/Source/Public/IAssetIndexer.h
@@ -7,7 +7,7 @@
 class UObject;
 class FSearchSerializer;
 
-class IAssetIndexer
+class ASSETSEARCH_API IAssetIndexer
 {
 public:
 
diff --git a/Engine/Plugins/Editor/AssetSearch/Source/Public/Utility/IndexerUtilities.h b/Engine/Plugins/Editor/AssetSearch/Source/Public/Utility/IndexerUtilities.h
--- a/Engine/Plugins/Editor/AssetSearch/Source/Public/Utility/IndexerUtilities.h
+++ b/Engine/Plugins/Editor/AssetSearch/Source/Public/Utility/IndexerUtilities.h
@@ -4,7 +4,7 @@
 
 #include "CoreMinimal.h"
 
-class FIndexerUtilities
+class ASSETSEARCH_API FIndexerUtilities
 {
 public:
 	static void IterateIndexableProperties(const UObject* InObject, TFunctionRef<void(const FProperty* /*Property*/, const FString& /*Value*/)> Callback);
