PR #9175: Fixed Asset Search indexing for FSoftObjectPtr representing references to non-assets (i.e. AActor instances) (Contributed by [REDACTED_PERSON])

diff --git a/Engine/Plugins/Editor/AssetSearch/Source/Private/Utility/IndexerUtilities.cpp b/Engine/Plugins/Editor/AssetSearch/Source/Private/Utility/IndexerUtilities.cpp
--- a/Engine/Plugins/Editor/AssetSearch/Source/Private/Utility/IndexerUtilities.cpp
+++ b/Engine/Plugins/Editor/AssetSearch/Source/Private/Utility/IndexerUtilities.cpp
@@ -112,7 +112,15 @@ void FIndexerUtilities::IterateIndexableProperties(const UStruct* InStruct, cons
 			FSoftObjectPtr SoftObject = SoftObjectProperty->GetPropertyValue(ValuePtr);
 			if (!SoftObject.IsNull())
 			{
+				const FSoftObjectPath& SoftObjectPath = SoftObject.ToSoftObjectPath();
 				Text = SoftObject.GetAssetName();
+
+				// If Soft Object Path is reference to AActor instance in the world, GetSubPathString() returns this actor instance path
+				// GetAssetName() would just return name of the level name, which won't provide valuable info for the search
+				if (!SoftObjectPath.GetSubPathString().IsEmpty())
+				{
+					Text.Append(TEXT(".")).Append(SoftObjectPath.GetSubPathString());
+				}	
 			}
 		}
 		else if (const FStructProperty* StructProperty = CastField<FStructProperty>(Property))
