Super Search - Include enumbyte properties.

diff --git a/Engine/Plugins/Editor/AssetSearch/Source/Private/Utility/IndexerUtilities.cpp b/Engine/Plugins/Editor/AssetSearch/Source/Private/Utility/IndexerUtilities.cpp
--- a/Engine/Plugins/Editor/AssetSearch/Source/Private/Utility/IndexerUtilities.cpp
+++ b/Engine/Plugins/Editor/AssetSearch/Source/Private/Utility/IndexerUtilities.cpp
@@ -76,6 +76,15 @@ void FIndexerUtilities::IterateIndexableProperties(const UStruct* InStruct, cons
 			const FText Value = TextProperty->GetPropertyValue(ValuePtr);
 			Text = *FTextInspector::GetSourceString(Value);
 		}
+		else if (const FByteProperty* ByteProperty = CastField<FByteProperty>(Property))
+		{
+			if (const UEnum* Enum = ByteProperty->Enum)
+			{
+				const int64 Value = ByteProperty->GetSignedIntPropertyValue(ValuePtr);
+				FText DisplayName = Enum->GetDisplayNameTextByValue(Value);
+				Text = *FTextInspector::GetSourceString(DisplayName);	
+			}
+		}
 		else if (const FEnumProperty* EnumProperty = CastField<FEnumProperty>(Property))
 		{
 			const int64 Value = EnumProperty->GetUnderlyingProperty()->GetSignedIntPropertyValue(ValuePtr);
