AssetSearch - Upgrade search db to store and restore full class paths, rather than partial paths.  Bumping version to recognize new change.

diff --git a/Engine/Plugins/Editor/AssetSearch/Source/Private/AssetSearchDatabase.cpp b/Engine/Plugins/Editor/AssetSearch/Source/Private/AssetSearchDatabase.cpp
--- a/Engine/Plugins/Editor/AssetSearch/Source/Private/AssetSearchDatabase.cpp
+++ b/Engine/Plugins/Editor/AssetSearch/Source/Private/AssetSearchDatabase.cpp
@@ -31,6 +31,7 @@ enum class EAssetSearchDatabaseVersion
 	IndexingAssetIdsAssetPathsUnique,
 	IntroducingFileHashing,
 	MovingFileInfoIntoAnotherDatabase,
+	ClassPathIsNowFullPath,
 
 	// -----<new versions can be added above this line>-------------------------------------------------
 	VersionPlusOne,
@@ -336,7 +337,7 @@ public:
 			FString ResultClassPath;
 			if (InStatement.GetColumnValues(Result.AssetName, ResultClassPath, Result.AssetPath))
 			{
-				Result.AssetClass = FTopLevelAssetPath(ResultClassPath);
+				Result.AssetClass.TrySetPath(ResultClassPath);
 				const float WorstCase = Result.AssetName.Len() + QueryText.Len();
 				Result.Score = -50.0f * (1.0f - (Algo::LevenshteinDistance(Result.AssetName.ToLower(), QueryText.ToLower()) / WorstCase));
 				
@@ -382,7 +383,7 @@ public:
 				Result.value_text, Result.value_hidden,
 				Result.Score))
 			{
-				Result.AssetClass = FTopLevelAssetPath(ResultClassPath);
+				Result.AssetClass.TrySetPath(ResultClassPath);
 				return InCallback(MoveTemp(Result));
 			}
 			return ESQLitePreparedStatementExecuteRowResult::Error;
