Removing commented out code.

diff --git a/Engine/Plugins/Editor/AssetSearch/Source/Private/Widgets/SSearchTreeRow.cpp b/Engine/Plugins/Editor/AssetSearch/Source/Private/Widgets/SSearchTreeRow.cpp
--- a/Engine/Plugins/Editor/AssetSearch/Source/Private/Widgets/SSearchTreeRow.cpp
+++ b/Engine/Plugins/Editor/AssetSearch/Source/Private/Widgets/SSearchTreeRow.cpp
@@ -56,14 +56,11 @@ TSharedRef<SWidget> SSearchTreeRow::GenerateWidgetForColumn(const FName& ColumnN
 
 			if (UClass* ObjectClass = UClass::TryFindTypeSlow<UClass>(*ObjectNode->object_native_class, EFindFirstObjectOptions::ExactClass))
 			{
-				//TSharedPtr<IAssetTypeActions> AssetActions = AssetToolsModule.Get().GetAssetTypeActionsForClass(ObjectClass).Pin();
-
 				FSlateIcon ClassIcon = FSlateIconFinder::FindIconForClass(ObjectClass);
 				if (ClassIcon.IsSet())
 				{
 					IconWidget = SNew(SImage)
 						.Image(ClassIcon.GetIcon());
-						//.ColorAndOpacity(AssetActions.IsValid() ? AssetActions->GetTypeColor() : FColor::White);
 				}
 			}
 		}
