AssetSearch - Asset search now allows you to choose on a project basis where to store the intermediate search data.  You can store it either in the DDC (which is where it has been going), or you can store the asset tag data.  The asset tag data might be preferrable if you don't have a large shared DDC by all users.  Storing it in the tag data aleviates that issue and makes it so the search data always travels with the asset, at the cost of bulking out the size of an asset a little more.

diff --git a/Engine/Plugins/Editor/AssetSearch/Source/Private/AssetSearchManager.cpp b/Engine/Plugins/Editor/AssetSearch/Source/Private/AssetSearchManager.cpp
--- a/Engine/Plugins/Editor/AssetSearch/Source/Private/AssetSearchManager.cpp
+++ b/Engine/Plugins/Editor/AssetSearch/Source/Private/AssetSearchManager.cpp
@@ -169,6 +169,8 @@ private:
 	TArray<UClass*> ClassFilters;
 };
 
+const FName FAssetSearchManager::AssetSearchIndexKeyTag(TEXT("AssetSearchIndexKey"));
+const FName FAssetSearchManager::AssetSearchIndexDataTag(TEXT("AssetSearchIndexData"));
 
 FAssetSearchManager::FAssetSearchManager()
 {
@@ -179,6 +181,8 @@ FAssetSearchManager::FAssetSearchManager()
 	TotalSearchRecords = 0;
 	LastRecordCountUpdateSeconds = 0;
 
+	IntermediateStorage = GetDefault<USearchProjectSettings>()->IntermediateStorage;
+
 	RunThread = false;
 }
 
@@ -195,7 +199,7 @@ FAssetSearchManager::~FAssetSearchManager()
 
 	UPackage::PackageSavedWithContextEvent.RemoveAll(this);
 	FCoreUObjectDelegates::OnAssetLoaded.RemoveAll(this);
-	UObject::FAssetRegistryTag::OnGetExtraObjectTags.RemoveAll(this);
+	UObject::FAssetRegistryTag::OnGetExtendedAssetRegistryTagsForSave.RemoveAll(this);
 
 	FTSTicker::GetCoreTicker().RemoveTicker(TickerHandle);
 }
@@ -218,13 +222,22 @@ void FAssetSearchManager::Start()
 
 	RegisterSearchProvider(TEXT("AssetRegistry"), MakeUnique<FAssetRegistrySearchProvider>());
 
-	UPackage::PackageSavedWithContextEvent.AddRaw(this, &FAssetSearchManager::HandlePackageSaved);
+	if (IntermediateStorage == ESearchIntermediateStorage::DerivedDataCache)
+	{
+		UPackage::PackageSavedWithContextEvent.AddRaw(this, &FAssetSearchManager::HandlePackageSaved);
+	}
+
 	FCoreUObjectDelegates::OnAssetLoaded.AddRaw(this, &FAssetSearchManager::OnAssetLoaded);
 
 	TickerHandle = FTSTicker::GetCoreTicker().AddTicker(FTickerDelegate::CreateRaw(this, &FAssetSearchManager::Tick_GameThread), 0);
 
 	RunThread = true;
 	DatabaseThread = FRunnableThread::Create(this, TEXT("UniversalSearch"), 0, TPri_BelowNormal);
+
+	if (IntermediateStorage == ESearchIntermediateStorage::AssetTagData)
+	{
+		UObject::FAssetRegistryTag::OnGetExtendedAssetRegistryTagsForSave.AddRaw(this, &FAssetSearchManager::HandleGetExtendedAssetRegistryTagsForSave);
+	}
 }
 
 void FAssetSearchManager::UpdateScanningAssets()
@@ -292,7 +305,7 @@ void FAssetSearchManager::StopScanningAssets()
 	}
 
 	ProcessAssetQueue.Reset();
-	FailedDDCRequests.Reset();
+	AssetNeedingReindexing.Reset();
 }
 
 void FAssetSearchManager::TryConnectToDatabase()
@@ -339,7 +352,7 @@ FSearchStats FAssetSearchManager::GetStats() const
 	Stats.Processing = IsAssetUpToDateCount + DownloadQueueCount + ActiveDownloads;
 	Stats.Updating = PendingDatabaseUpdates;
 	Stats.TotalRecords = TotalSearchRecords;
-	Stats.AssetsMissingIndex = FailedDDCRequests.Num();
+	Stats.AssetsMissingIndex = AssetNeedingReindexing.Num();
 
 	return Stats;
 }
@@ -362,11 +375,22 @@ void FAssetSearchManager::OnAssetAdded(const FAssetData& InAssetData)
 {
 	check(IsInGameThread());
 
+	static const FString EngineContentPathWithSlash = FPackageName::FilenameToLongPackageName(FPaths::EngineContentDir());
 	static const FString DeveloperPathWithSlash = FPackageName::FilenameToLongPackageName(FPaths::GameDevelopersDir());
 	static const FString UsersDeveloperPathWithSlash = FPackageName::FilenameToLongPackageName(FPaths::GameUserDeveloperDir());
+
+	const FString PackageName = InAssetData.PackageName.ToString();
+
+	// Don't index things in the engine directory if we're storing data in asset tags, since we don't want to mutate the engine data.
+	if (IntermediateStorage == ESearchIntermediateStorage::AssetTagData)
+	{
+		if (PackageName.StartsWith(EngineContentPathWithSlash))
+		{
+			return;
+		}
+	}
 	
 	// Don't process stuff in the other developer folders.
-	FString PackageName = InAssetData.PackageName.ToString();
 	if (PackageName.StartsWith(DeveloperPathWithSlash))
 	{
 		if (!PackageName.StartsWith(UsersDeveloperPathWithSlash))
@@ -375,7 +399,7 @@ void FAssetSearchManager::OnAssetAdded(const FAssetData& InAssetData)
 		}
 	}
 
-	// 
+	// Ignore it if the package is in one of the project ignored paths.
 	const USearchProjectSettings* ProjectSettings = GetDefault<USearchProjectSettings>();
 	for (const FDirectoryPath& IgnoredPath : ProjectSettings->IgnoredPaths)
 	{
@@ -385,7 +409,7 @@ void FAssetSearchManager::OnAssetAdded(const FAssetData& InAssetData)
 		}
 	}
 
-	// 
+	// Ignore it if the package is in one of the user ignored paths.
 	const USearchUserSettings* UserSettings = GetDefault<USearchUserSettings>();
 	for (const FDirectoryPath& IgnoredPath : UserSettings->IgnoredPaths)
 	{
@@ -442,9 +466,39 @@ void FAssetSearchManager::OnAssetScanFinished()
 	});
 }
 
+void FAssetSearchManager::HandleGetExtendedAssetRegistryTagsForSave(const UObject* Object, const ITargetPlatform* TargetPlatform, TArray<UObject::FAssetRegistryTag>& OutTags)
+{
+	if (GIsEditor && !Object->GetOutermost()->HasAnyPackageFlags(PKG_ForDiffing) && !IsRunningCookCommandlet())
+	{
+		if (!GEditor->IsAutosaving())
+		{
+			if (!TargetPlatform || TargetPlatform->HasEditorOnlyData())
+			{
+				FAssetData InAssetData(Object, FAssetData::ECreationFlags::SkipAssetRegistryTagsGathering);
+
+				FString IndexedJson;
+				const bool bWasIndexed = IndexAsset(InAssetData, Object, IndexedJson);
+				if (bWasIndexed && !IndexedJson.IsEmpty())
+				{
+					const FString IndexKey = GetBaseIndexKey(Object->GetClass());
+
+					OutTags.Add(UObject::FAssetRegistryTag(FAssetSearchManager::AssetSearchIndexKeyTag, IndexKey, UObject::FAssetRegistryTag::TT_Hidden));
+					OutTags.Add(UObject::FAssetRegistryTag(FAssetSearchManager::AssetSearchIndexDataTag, IndexedJson, UObject::FAssetRegistryTag::TT_Hidden));
+
+					if (!IndexedJson.IsEmpty())
+					{
+						AddOrUpdateAsset(InAssetData, IndexedJson, IndexKey);
+					}
+				}
+			}
+		}
+	}
+}
+
 void FAssetSearchManager::HandlePackageSaved(const FString& PackageFilename, UPackage* Package, FObjectPostSaveContext ObjectSaveContext)
 {
 	check(IsInGameThread());
+	check(IntermediateStorage == ESearchIntermediateStorage::DerivedDataCache);
 
 	// Only execute if this is a user save
 	if (ObjectSaveContext.IsProceduralSave())
@@ -459,7 +513,7 @@ void FAssetSearchManager::HandlePackageSaved(const FString& PackageFilename, UPa
 		GetObjectsWithPackage(Package, Objects, bIncludeNestedObjects);
 		for (UObject* Entry : Objects)
 		{
-			RequestIndexAsset(Entry);
+			RequestIndexAsset_DDC(Entry);
 		}
 	}
 }
@@ -470,13 +524,23 @@ void FAssetSearchManager::OnAssetLoaded(UObject* InObject)
 
 	if (bTryIndexAssetsOnLoad)
 	{
-		RequestIndexAsset(InObject);
+		switch (IntermediateStorage)
+		{
+			case ESearchIntermediateStorage::DerivedDataCache:
+				RequestIndexAsset_DDC(InObject);
+				break;
+			case ESearchIntermediateStorage::AssetTagData:
+				// Don't need to do this on load, I mean, it's a little more accurate but if we're storing it
+				// in tag data, it's almost certainly up to date unless someone changed the serialization.
+				break;
+		}
 	}
 }
 
-bool FAssetSearchManager::RequestIndexAsset(UObject* InAsset)
+bool FAssetSearchManager::RequestIndexAsset_DDC(const UObject* InAsset)
 {
 	check(IsInGameThread());
+	check(IntermediateStorage == ESearchIntermediateStorage::DerivedDataCache);
 
 	if (GEditor == nullptr || GEditor->IsAutosaving())
 	{
@@ -485,7 +549,7 @@ bool FAssetSearchManager::RequestIndexAsset(UObject* InAsset)
 
 	if (IsAssetIndexable(InAsset))
 	{
-		TWeakObjectPtr<UObject> AssetWeakPtr = InAsset;
+		TWeakObjectPtr<const UObject> AssetWeakPtr = InAsset;
 		FAssetData AssetData(InAsset);
 
 		return AsyncGetDerivedDataKey(AssetData, [this, AssetData, AssetWeakPtr](bool bSuccess, FString InDDCKey) {
@@ -499,7 +563,7 @@ bool FAssetSearchManager::RequestIndexAsset(UObject* InAsset)
 				if (!SearchDatabase.IsAssetUpToDate(AssetData, InDDCKey))
 				{
 					AsyncMainThreadTask([this, AssetWeakPtr]() {
-						StoreIndexForAsset(AssetWeakPtr.Get());
+						StoreIndexForAsset_DDC(AssetWeakPtr.Get());
 					});
 				}
 			});
@@ -509,7 +573,7 @@ bool FAssetSearchManager::RequestIndexAsset(UObject* InAsset)
 	return false;
 }
 
-bool FAssetSearchManager::IsAssetIndexable(UObject* InAsset)
+bool FAssetSearchManager::IsAssetIndexable(const UObject* InAsset) const
 {
 	if (InAsset && InAsset->IsAsset())
 	{
@@ -533,7 +597,75 @@ bool FAssetSearchManager::IsAssetIndexable(UObject* InAsset)
 
 bool FAssetSearchManager::TryLoadIndexForAsset(const FAssetData& InAssetData)
 {
-	bool bSuccess = AsyncGetDerivedDataKey(InAssetData, [this, InAssetData](bool bSuccess, FString InDDCKey) {
+	switch (IntermediateStorage)
+	{
+		case ESearchIntermediateStorage::DerivedDataCache:
+			return TryLoadIndexForAsset_DDC(InAssetData);
+		case ESearchIntermediateStorage::AssetTagData:
+			return TryLoadIndexForAsset_Tags(InAssetData);
+		default:
+			ensure(false);
+			return false;
+	}
+}
+
+bool FAssetSearchManager::TryLoadIndexForAsset_Tags(const FAssetData& InAssetData)
+{
+	check(IsInGameThread());
+	check(IntermediateStorage == ESearchIntermediateStorage::AssetTagData);
+
+	IsAssetUpToDateCount++;
+
+	FString CurrentIndexerNameAndVersion = GetBaseIndexKey(InAssetData.GetClass());
+	if (!CurrentIndexerNameAndVersion.IsEmpty())
+	{
+		const FString SavedIndexerNameAndVersion = InAssetData.GetTagValueRef<FString>(FAssetSearchManager::AssetSearchIndexKeyTag);
+		if (!SavedIndexerNameAndVersion.IsEmpty())
+		{
+			if (SavedIndexerNameAndVersion != CurrentIndexerNameAndVersion)
+			{
+				// Log need to update/missing...etc.
+				AssetNeedingReindexing.Add(InAssetData);
+			}
+
+			UpdateOperations.Enqueue([this, InAssetData, CurrentIndexerNameAndVersion]() {
+				FScopeLock ScopedLock(&SearchDatabaseCS);
+				if (!SearchDatabase.IsAssetUpToDate(InAssetData, CurrentIndexerNameAndVersion))
+				{
+					AsyncMainThreadTask([this, InAssetData, CurrentIndexerNameAndVersion]() {
+						const FString IndexedJson = InAssetData.GetTagValueRef<FString>(FAssetSearchManager::AssetSearchIndexDataTag);
+
+						IsAssetUpToDateCount--;
+
+						if (!IndexedJson.IsEmpty())
+						{
+							AddOrUpdateAsset(InAssetData, IndexedJson, CurrentIndexerNameAndVersion);
+						}
+					});
+				}
+				else
+				{
+					IsAssetUpToDateCount--;
+				}
+			});
+
+			return true;
+		}
+		else
+		{
+ 			AssetNeedingReindexing.Add(InAssetData);
+		}
+	}
+
+	IsAssetUpToDateCount--;
+	return false;
+}
+
+bool FAssetSearchManager::TryLoadIndexForAsset_DDC(const FAssetData& InAssetData)
+{
+	check(IntermediateStorage == ESearchIntermediateStorage::DerivedDataCache);
+
+	const bool bSuccess = AsyncGetDerivedDataKey(InAssetData, [this, InAssetData](bool bSuccess, FString InDDCKey) {
 		if (!bSuccess)
 		{
 			IsAssetUpToDateCount--;
@@ -572,6 +704,7 @@ void FAssetSearchManager::AsyncRequestDownload(const FAssetData& InAssetData, co
 bool FAssetSearchManager::AsyncGetDerivedDataKey(const FAssetData& InAssetData, TFunction<void(bool, FString)> DDCKeyCallback)
 {
 	check(IsInGameThread());
+	check(IntermediateStorage == ESearchIntermediateStorage::DerivedDataCache);
 
 	const FString AssetPath = InAssetData.PackagePath.ToString();
 	if (AssetPath.Contains(FPackagePath::GetExternalActorsFolderName()) || AssetPath.Contains(FPackagePath::GetExternalObjectsFolderName()))
@@ -645,6 +778,22 @@ bool FAssetSearchManager::HasIndexerForClass(const UClass* InAssetClass) const
 	return false;
 }
 
+FString FAssetSearchManager::GetBaseIndexKey(const UClass* InAssetClass) const
+{
+	TStringBuilder<256> VersionString;
+
+	const FString AssetIndexerVersions = GetIndexerVersion(InAssetClass);
+	if (!AssetIndexerVersions.IsEmpty())
+	{
+		VersionString.Append(TEXT("AssetSearch_V"));
+		VersionString.Append(LexToString(FSearchSerializer::GetVersion()));
+		VersionString.Append(TEXT("_"));
+		VersionString.Append(AssetIndexerVersions);
+	}
+
+	return VersionString.ToString();
+}
+
 FString FAssetSearchManager::GetIndexerVersion(const UClass* InAssetClass) const
 {
 	TStringBuilder<256> VersionString;
@@ -675,40 +824,70 @@ FString FAssetSearchManager::GetIndexerVersion(const UClass* InAssetClass) const
 	return VersionString.ToString();
 }
 
-void FAssetSearchManager::StoreIndexForAsset(UObject* InAsset)
+bool FAssetSearchManager::IndexAsset(const FAssetData& InAssetData, const UObject* InAsset, FString& OutIndexedJson) const
 {
 	check(IsInGameThread());
 
 	if (IsAssetIndexable(InAsset) && HasIndexerForClass(InAsset->GetClass()))
 	{
-		FAssetData InAssetData(InAsset);
+		FSearchSerializer Serializer(InAssetData, &OutIndexedJson);
+		return Serializer.IndexAsset(InAsset, Indexers);
+	}
 
-		FString IndexedJson;
-		bool bWasIndexed = false;
-		{
-			FSearchSerializer Serializer(InAssetData, &IndexedJson);
-			bWasIndexed = Serializer.IndexAsset(InAsset, Indexers);
-		}
+	return false;
+}
 
-		if (bWasIndexed && !IndexedJson.IsEmpty())
-		{
-			AsyncGetDerivedDataKey(InAssetData, [this, InAssetData, IndexedJson](bool bSuccess, FString InDDCKey) {
-				if (!bSuccess)
-				{
-					return;
-				}
+void FAssetSearchManager::StoreIndexForAsset(const UObject* InAsset)
+{
+	switch (IntermediateStorage)
+	{
+	case ESearchIntermediateStorage::DerivedDataCache:
+		return StoreIndexForAsset_DDC(InAsset);
+	case ESearchIntermediateStorage::AssetTagData:
+		return StoreIndexForAsset_Tags(InAsset);
+	default:
+		ensure(false);
+	}
+}
 
-				AsyncMainThreadTask([this, InAssetData, IndexedJson, InDDCKey]() {
-					check(IsInGameThread());
+void FAssetSearchManager::StoreIndexForAsset_Tags(const UObject* InAsset)
+{
+	check(IsInGameThread());
+	check(IntermediateStorage == ESearchIntermediateStorage::AssetTagData);
 
-					FTCHARToUTF8 IndexedJsonUTF8(*IndexedJson);
-					TArrayView<const uint8> IndexedJsonUTF8View((const uint8*)IndexedJsonUTF8.Get(), IndexedJsonUTF8.Length() * sizeof(UTF8CHAR));
-					GetDerivedDataCacheRef().Put(*InDDCKey, IndexedJsonUTF8View, InAssetData.GetObjectPathString(), false);
+	// StoreIndexForAsset doesn't really do much of anything when it's stored in tag data, all we can really
+	// do is mark the package dirty so that it can be saved.  There's really nothing else that need be done
+	// to upgrade it, since the index data was already grabbed from the tag data upon discovery.
+	InAsset->MarkPackageDirty();
+}
 
-					AddOrUpdateAsset(InAssetData, IndexedJson, InDDCKey);
-				});
+void FAssetSearchManager::StoreIndexForAsset_DDC(const UObject* InAsset)
+{
+	check(IsInGameThread());
+	check(IntermediateStorage == ESearchIntermediateStorage::DerivedDataCache);
+
+	FString IndexedJson;
+	FAssetData InAssetData(InAsset, FAssetData::ECreationFlags::SkipAssetRegistryTagsGathering);
+	const bool bWasIndexed = IndexAsset(InAssetData, InAsset, IndexedJson);
+
+	if (bWasIndexed && !IndexedJson.IsEmpty())
+	{
+		AsyncGetDerivedDataKey(InAssetData, [this, InAssetData, IndexedJson](bool bSuccess, FString InDDCKey) {
+			if (!bSuccess)
+			{
+				return;
+			}
+
+			AsyncMainThreadTask([this, InAssetData, IndexedJson, InDDCKey]() {
+				check(IsInGameThread());
+
+				FTCHARToUTF8 IndexedJsonUTF8(*IndexedJson);
+				TArrayView<const uint8> IndexedJsonUTF8View((const uint8*)IndexedJsonUTF8.Get(), IndexedJsonUTF8.Length() * sizeof(UTF8CHAR));
+				GetDerivedDataCacheRef().Put(*InDDCKey, IndexedJsonUTF8View, InAssetData.GetObjectPathString(), false);
+
+				AddOrUpdateAsset(InAssetData, IndexedJson, InDDCKey);
 			});
-		}
+		});
 	}
 }
 
@@ -797,9 +976,9 @@ bool FAssetSearchManager::Tick_GameThread(float DeltaTime)
 				LoadDDCContentIntoDatabase(PendingRequest->AssetData, OutContent, PendingRequest->DDCKey);
 				DownloadProcessLimit--;
 			}
-			else if (UserSettings->bShowMissingAssets)
+			else if (UserSettings->bShowAssetsNeedingIndexing)
 			{
-				FailedDDCRequests.Add(*PendingRequest);
+				AssetNeedingReindexing.Add(PendingRequest->AssetData);
 			}
 
 			ProcessDDCQueue.Pop();
@@ -856,9 +1035,7 @@ void FAssetSearchManager::ForceIndexOnAssetsMissingIndex()
 {
 	check(IsInGameThread());
 
-	EAppReturnType::Type IncludeMaps = FMessageDialog::Open(EAppMsgType::YesNo, LOCTEXT("IncludeMaps", "Do you want to open and index map files, this can take a long time?"));
-
-	FScopedSlowTask IndexingTask(FailedDDCRequests.Num(), LOCTEXT("ForceIndexOnAssetsMissingIndex", "Indexing Assets"));
+	FScopedSlowTask IndexingTask(AssetNeedingReindexing.Num(), LOCTEXT("ForceIndexOnAssetsMissingIndex", "Indexing Assets"));
 	IndexingTask.MakeDialog(true);
 
 	int32 RemovedCount = 0;
@@ -867,11 +1044,11 @@ void FAssetSearchManager::ForceIndexOnAssetsMissingIndex()
 
 	TGuardValue<bool> GuardResetTesting(GIsAutomationTesting, true);
 
-	const int32 OnePercentChunk = (int32)(FailedDDCRequests.Num() / 100.0);
+	const int32 OnePercentChunk = (int32)(AssetNeedingReindexing.Num() / 100.0);
 	int32 ChunkProgress = 0;
 
 	FUnloadPackageScope UnloadScope;
-	for (const FAssetDDCRequest& Request : FailedDDCRequests)
+	for (const FAssetData& Asset : AssetNeedingReindexing)
 	{
 		if (IndexingTask.ShouldCancel())
 		{
@@ -881,27 +1058,16 @@ void FAssetSearchManager::ForceIndexOnAssetsMissingIndex()
 		if (RemovedCount > ChunkProgress)
 		{
 			ChunkProgress += OnePercentChunk;
-			IndexingTask.EnterProgressFrame(OnePercentChunk, FText::Format(LOCTEXT("ForceIndexOnAssetsMissingIndexFormat", "Indexing Asset ({0} of {1})"), RemovedCount + 1, FailedDDCRequests.Num()));
-		}
-
-		if (IncludeMaps != EAppReturnType::Yes)
-		{
-			if (Request.AssetData.GetClass() == UWorld::StaticClass())
-			{
-				RemovedCount++;
-				continue;
-			}
+			IndexingTask.EnterProgressFrame(OnePercentChunk, FText::Format(LOCTEXT("ForceIndexOnAssetsMissingIndexFormat", "Indexing Asset ({0} of {1})"), RemovedCount + 1, AssetNeedingReindexing.Num()));
 		}
 
-		//ProcessGameThreadTasks();
-
-		if (UObject* AssetToIndex = Request.AssetData.GetAsset())
+		if (UObject* AssetToIndex = Asset.GetAsset())
 		{
 			// This object's metadata incorrectly labled it as something other than a redirector.  We need to resave it
 			// to stop it from appearing as something it's not.
 			if (UObjectRedirector* Redirector = Cast<UObjectRedirector>(AssetToIndex))
 			{
-				RedirectorsWithBrokenMetadata.Add(Request.AssetData);
+				RedirectorsWithBrokenMetadata.Add(Asset);
 				RemovedCount++;
 				continue;
 			}
@@ -940,7 +1106,18 @@ void FAssetSearchManager::ForceIndexOnAssetsMissingIndex()
 		}
 	}
 
-	FailedDDCRequests.RemoveAtSwap(0, RemovedCount);
+	AssetNeedingReindexing.RemoveAtSwap(0, RemovedCount);
+
+	if (IntermediateStorage == ESearchIntermediateStorage::AssetTagData)
+	{
+		const bool bPromptUserToSave = true;
+		const bool bSaveMapPackages = true;
+		const bool bSaveContentPackages = true;
+		const bool bFastSave = false;
+		const bool bNotifyNoPackagesSaved = false;
+		const bool bCanBeDeclined = false;
+		FEditorFileUtils::SaveDirtyPackages(bPromptUserToSave, bSaveMapPackages, bSaveContentPackages, bFastSave, bNotifyNoPackagesSaved, bCanBeDeclined);
+	}
 }
 
 void FAssetSearchManager::Search(FSearchQueryPtr SearchQuery)
diff --git a/Engine/Plugins/Editor/AssetSearch/Source/Private/AssetSearchManager.h b/Engine/Plugins/Editor/AssetSearch/Source/Private/AssetSearchManager.h
--- a/Engine/Plugins/Editor/AssetSearch/Source/Private/AssetSearchManager.h
+++ b/Engine/Plugins/Editor/AssetSearch/Source/Private/AssetSearchManager.h
@@ -18,9 +18,13 @@ class FObjectPostSaveContext;
 class FRunnableThread;
 class UClass;
 class ISearchProvider;
+enum class ESearchIntermediateStorage : uint8;
 
 class FAssetSearchManager : public FRunnable
 {
+	static const FName AssetSearchIndexKeyTag;
+	static const FName AssetSearchIndexDataTag;
+
 public:
 	FAssetSearchManager();
 	virtual ~FAssetSearchManager();
@@ -52,19 +56,31 @@ private:
 	void OnAssetRemoved(const FAssetData& InAssetData);
 	void OnAssetScanFinished();
 
+	void HandleGetExtendedAssetRegistryTagsForSave(const UObject* Object, const ITargetPlatform* TargetPlatform, TArray<UObject::FAssetRegistryTag>& OutTags);
 	void HandlePackageSaved(const FString& PackageFilename, UPackage* Package, FObjectPostSaveContext ObjectSaveContext);
 	void OnAssetLoaded(UObject* InObject);
 
 	void AddOrUpdateAsset(const FAssetData& InAsset, const FString& IndexedJson, const FString& DerivedDataKey);
 
-	bool RequestIndexAsset(UObject* InAsset);
-	bool IsAssetIndexable(UObject* InAsset);
+	bool RequestIndexAsset_DDC(const UObject* InAsset);
+	bool IsAssetIndexable(const UObject* InAsset) const;
+	
 	bool TryLoadIndexForAsset(const FAssetData& InAsset);
+	bool TryLoadIndexForAsset_Tags(const FAssetData& InAssetData);
+	bool TryLoadIndexForAsset_DDC(const FAssetData& InAssetData);
+
 	void AsyncRequestDownload(const FAssetData& InAssetData, const FString& InDDCKey);
 	bool AsyncGetDerivedDataKey(const FAssetData& UnindexedAsset, TFunction<void(bool, FString)> DDCKeyCallback);
 	bool HasIndexerForClass(const UClass* InAssetClass) const;
+
+	FString GetBaseIndexKey(const UClass* InAssetClass) const;
 	FString GetIndexerVersion(const UClass* InAssetClass) const;
-	void StoreIndexForAsset(UObject* InAsset);
+	bool IndexAsset(const FAssetData& InAssetData, const UObject* InAsset, FString& OutIndexedJson) const;
+	
+	void StoreIndexForAsset(const UObject* InAsset);
+	void StoreIndexForAsset_Tags(const UObject* InAsset);
+	void StoreIndexForAsset_DDC(const UObject* InAsset);
+
 	void LoadDDCContentIntoDatabase(const FAssetData& InAsset, const TArray<uint8>& Content, const FString& DerivedDataKey);
 
 	void AsyncMainThreadTask(TFunction<void()> Task);
@@ -110,7 +126,7 @@ private:
 	TQueue<FAssetDDCRequest, EQueueMode::Mpsc> DownloadQueue;
 	TQueue<FAssetDDCRequest, EQueueMode::Mpsc> ProcessDDCQueue;
 
-	TArray<FAssetDDCRequest> FailedDDCRequests;
+	TArray<FAssetData> AssetNeedingReindexing;
 
 	FTSTicker::FDelegateHandle TickerHandle;
 
@@ -120,6 +136,8 @@ private:
 	volatile bool bDatabaseOpen = false;
 	volatile double LastConnectionAttempt = 0;
 
+	ESearchIntermediateStorage IntermediateStorage;
+
 	TAtomic<bool> RunThread;
 	FRunnableThread* DatabaseThread = nullptr;
 
diff --git a/Engine/Plugins/Editor/AssetSearch/Source/Private/SearchSerializer.cpp b/Engine/Plugins/Editor/AssetSearch/Source/Private/SearchSerializer.cpp
--- a/Engine/Plugins/Editor/AssetSearch/Source/Private/SearchSerializer.cpp
+++ b/Engine/Plugins/Editor/AssetSearch/Source/Private/SearchSerializer.cpp
@@ -67,7 +67,7 @@ int32 FSearchSerializer::GetVersion()
 	return (int32)ESearchIndexVersion::LatestVersion;
 }
 
-bool FSearchSerializer::IndexAsset(UObject* InAsset, const TMap<FName, TUniquePtr<IAssetIndexer>>& Indexers)
+bool FSearchSerializer::IndexAsset(const UObject* InAsset, const TMap<FName, TUniquePtr<IAssetIndexer>>& Indexers)
 {
 	if (InAsset == nullptr)
 	{
diff --git a/Engine/Plugins/Editor/AssetSearch/Source/Private/Settings/SearchProjectSettings.h b/Engine/Plugins/Editor/AssetSearch/Source/Private/Settings/SearchProjectSettings.h
--- a/Engine/Plugins/Editor/AssetSearch/Source/Private/Settings/SearchProjectSettings.h
+++ b/Engine/Plugins/Editor/AssetSearch/Source/Private/Settings/SearchProjectSettings.h
@@ -7,6 +7,16 @@
 
 struct FDirectoryPath;
 
+UENUM()
+enum class ESearchIntermediateStorage : uint8
+{
+	// Stores the json snippets of indexable data in the DDC so that they can be downloaded as needed.
+	DerivedDataCache,
+	// Stores the json snippets and the json hash in the tag data for the asset.  This increases the size of all 
+	// assets that are indexed, but can be a great benefit to teams that don't have a large shared DDC.
+	AssetTagData
+};
+
 UCLASS(config = Editor, defaultconfig, meta=(DisplayName="Search"))
 class USearchProjectSettings : public UDeveloperSettings
 {
@@ -15,6 +25,9 @@ class USearchProjectSettings : public UDeveloperSettings
 public:
 	USearchProjectSettings();
 
+	UPROPERTY(config, EditAnywhere, Category=General, Meta = (ConfigRestartRequired = true))
+	ESearchIntermediateStorage IntermediateStorage = ESearchIntermediateStorage::DerivedDataCache;
+
 	UPROPERTY(config, EditAnywhere, Category=General)
 	TArray<FDirectoryPath> IgnoredPaths;
 };
diff --git a/Engine/Plugins/Editor/AssetSearch/Source/Private/Settings/SearchUserSettings.h b/Engine/Plugins/Editor/AssetSearch/Source/Private/Settings/SearchUserSettings.h
--- a/Engine/Plugins/Editor/AssetSearch/Source/Private/Settings/SearchUserSettings.h
+++ b/Engine/Plugins/Editor/AssetSearch/Source/Private/Settings/SearchUserSettings.h
@@ -39,7 +39,7 @@ public:
 	TArray<FDirectoryPath> IgnoredPaths;
 
 	UPROPERTY(config, EditAnywhere, Category=General)
-	bool bShowMissingAssets = true;
+	bool bShowAssetsNeedingIndexing = true;
 
 	UPROPERTY(config, EditAnywhere, Category = General)
 	bool bAutoExpandAssets = true;
diff --git a/Engine/Plugins/Editor/AssetSearch/Source/Private/Widgets/SSearchBrowser.cpp b/Engine/Plugins/Editor/AssetSearch/Source/Private/Widgets/SSearchBrowser.cpp
--- a/Engine/Plugins/Editor/AssetSearch/Source/Private/Widgets/SSearchBrowser.cpp
+++ b/Engine/Plugins/Editor/AssetSearch/Source/Private/Widgets/SSearchBrowser.cpp
@@ -165,7 +165,8 @@ void SSearchBrowser::Construct( const FArguments& InArgs )
 				[
 					SNew(SHyperlink)
 					.Text(this, &SSearchBrowser::GetUnindexedAssetsText)
-					.Visibility_Lambda([] { return GetDefault<USearchUserSettings>()->bShowMissingAssets ? EVisibility::Visible : EVisibility::Collapsed; })
+					.ToolTipText(LOCTEXT("AssetsNeedIndexingTooltip", "Click this to open and index the assets that are don't have any index data or their index data was found to be out of date."))
+					.Visibility_Lambda([] { return GetDefault<USearchUserSettings>()->bShowAssetsNeedingIndexing ? EVisibility::Visible : EVisibility::Collapsed; })
 					.OnNavigate(this, &SSearchBrowser::HandleForceIndexOfAssetsMissingIndex)
 				]
 
@@ -290,7 +291,7 @@ FText SSearchBrowser::GetUnindexedAssetsText() const
 
 	if (SearchStats.AssetsMissingIndex > 0)
 	{
-		return FText::Format(LOCTEXT("UnindexedAssetsLinkFormat", "{0} Missing"), SearchStats.AssetsMissingIndex);
+		return FText::Format(LOCTEXT("UnindexedAssetsLinkFormat", "{0} Need Indexing"), SearchStats.AssetsMissingIndex);
 	}
 
 	return FText::GetEmpty();
diff --git a/Engine/Plugins/Editor/AssetSearch/Source/Public/SearchSerializer.h b/Engine/Plugins/Editor/AssetSearch/Source/Public/SearchSerializer.h
--- a/Engine/Plugins/Editor/AssetSearch/Source/Public/SearchSerializer.h
+++ b/Engine/Plugins/Editor/AssetSearch/Source/Public/SearchSerializer.h
@@ -18,7 +18,7 @@ public:
 
 	static int32 GetVersion();
 
-	bool IndexAsset(UObject* InAsset, const TMap<FName, TUniquePtr<IAssetIndexer>>& Indexers);
+	bool IndexAsset(const UObject* InAsset, const TMap<FName, TUniquePtr<IAssetIndexer>>& Indexers);
 
 	void BeginIndexer(const IAssetIndexer* InIndexer);
 
