Original PR modified - removing the Ctrl+F binding for now, but keeping the ability to bind it.

PR #9189: Bind Asset Search to Ctrl+F hotkey. (Contributed by [REDACTED_PERSON])

diff --git a/Engine/Plugins/Editor/AssetSearch/Source/Private/AssetSearchCommands.cpp b/Engine/Plugins/Editor/AssetSearch/Source/Private/AssetSearchCommands.cpp
new file mode 100644
--- /dev/null
+++ b/Engine/Plugins/Editor/AssetSearch/Source/Private/AssetSearchCommands.cpp
@@ -0,0 +1,20 @@
+// Copyright Epic Games, Inc. All Rights Reserved.
+
+#include "AssetSearchCommands.h"
+
+#define LOCTEXT_NAMESPACE "AssetSearchCommands"
+
+FAssetSearchCommands::FAssetSearchCommands() : TCommands<FAssetSearchCommands>(
+	"AssetSearchCommands",
+	NSLOCTEXT("Contexts", "AssetSearchCommands", "Asset Search"),
+	NAME_None, 
+	FAppStyle::GetAppStyleSetName())
+{
+}
+
+void FAssetSearchCommands::RegisterCommands()
+{
+	UI_COMMAND(ViewAssetSearch, "Search...", "Show Search Tab", EUserInterfaceActionType::Button, FInputChord());
+}
+
+#undef LOCTEXT_NAMESPACE
diff --git a/Engine/Plugins/Editor/AssetSearch/Source/Private/AssetSearchModule.cpp b/Engine/Plugins/Editor/AssetSearch/Source/Private/AssetSearchModule.cpp
--- a/Engine/Plugins/Editor/AssetSearch/Source/Private/AssetSearchModule.cpp
+++ b/Engine/Plugins/Editor/AssetSearch/Source/Private/AssetSearchModule.cpp
@@ -5,7 +5,11 @@
 #include "WorkspaceMenuStructureModule.h"
 #include "Widgets/Docking/SDockTab.h"
 #include "Widgets/SSearchBrowser.h"
+#include "AssetSearchCommands.h"
 #include "AssetSearchManager.h"
+#include "ContentBrowserDelegates.h"
+#include "ContentBrowserModule.h"
+#include "LevelEditor.h"
 
 #define LOCTEXT_NAMESPACE "FAssetSearchModule"
 
@@ -13,11 +17,6 @@ static const FName SearchTabName("Search");
 
 class FAssetSearchModule : public IAssetSearchModule
 {
-public:
-
-	// IAssetSearchModule interface
-
-
 public:
 
 	// IModuleInterface interface
@@ -27,16 +26,51 @@ public:
 		SearchManager = MakeUnique<FAssetSearchManager>();
 		SearchManager->Start();
 
-		FGlobalTabmanager::Get()->RegisterNomadTabSpawner(SearchTabName, FOnSpawnTab::CreateRaw(this, &FAssetSearchModule::HandleSpawnSettingsTab))
-			.SetDisplayName(LOCTEXT("Search", "Search"))
-			.SetTooltipText(LOCTEXT("SearchTab", "Search Tab"))
-			.SetIcon(FSlateIcon(FAppStyle::GetAppStyleSetName(), "Symbols.SearchGlass"))
-			.SetGroup(WorkspaceMenu::GetMenuStructure().GetToolsCategory());
+		if((GIsEditor && !IsRunningCommandlet()))
+		{
+			FAssetSearchCommands::Register();
+			
+			FGlobalTabmanager::Get()->RegisterNomadTabSpawner(SearchTabName, FOnSpawnTab::CreateRaw(this, &FAssetSearchModule::HandleSpawnSettingsTab))
+				.SetDisplayName(LOCTEXT("Search", "Search"))
+				.SetTooltipText(LOCTEXT("SearchTab", "Search Tab"))
+				.SetIcon(FSlateIcon(FAppStyle::GetAppStyleSetName(), "Symbols.SearchGlass"))
+				.SetGroup(WorkspaceMenu::GetMenuStructure().GetToolsCategory());
+				
+			// Register content browser hook
+			FContentBrowserModule& ContentBrowserModule = FModuleManager::LoadModuleChecked<FContentBrowserModule>(TEXT("ContentBrowser"));
+		
+			TArray<FContentBrowserCommandExtender>& CBCommandExtenderDelegates = ContentBrowserModule.GetAllContentBrowserCommandExtenders();
+			CBCommandExtenderDelegates.Add(FContentBrowserCommandExtender::CreateRaw(this, &FAssetSearchModule::OnExtendContentBrowserCommands));
+			ContentBrowserCommandExtenderDelegateHandle = CBCommandExtenderDelegates.Last().GetHandle();
+			
+			// Register level editor hooks and commands
+			FLevelEditorModule& LevelEditorModule = FModuleManager::Get().LoadModuleChecked<FLevelEditorModule>("LevelEditor");
+			TSharedRef<FUICommandList> CommandList = LevelEditorModule.GetGlobalLevelEditorActions();
+			
+			CommandList->MapAction(FAssetSearchCommands::Get().ViewAssetSearch, 
+			FExecuteAction::CreateLambda([this]
+			{
+				ExecuteOpenObjectBrowser();
+			}));
+		}
 	}
 
 	virtual void ShutdownModule() override
 	{
 		FGlobalTabmanager::Get()->UnregisterNomadTabSpawner(SearchTabName);
+
+		if ((GIsEditor && !IsRunningCommandlet()) && UObjectInitialized() && FSlateApplication::IsInitialized())
+		{
+			FContentBrowserModule& ContentBrowserModule = FModuleManager::LoadModuleChecked<FContentBrowserModule>(TEXT("ContentBrowser"));
+		
+			TArray<FContentBrowserCommandExtender>& CBCommandExtenderDelegates = ContentBrowserModule.GetAllContentBrowserCommandExtenders();
+			CBCommandExtenderDelegates.RemoveAll([this](const FContentBrowserCommandExtender& Delegate) { return Delegate.GetHandle() == ContentBrowserCommandExtenderDelegateHandle; });
+			
+			FLevelEditorModule& LevelEditorModule = FModuleManager::Get().LoadModuleChecked<FLevelEditorModule>("LevelEditor");
+
+			TSharedRef<FUICommandList> CommandList = LevelEditorModule.GetGlobalLevelEditorActions();
+			CommandList->UnmapAction(FAssetSearchCommands::Get().ViewAssetSearch);
+		}
 	}
 
 	void ExecuteOpenObjectBrowser()
@@ -44,6 +78,16 @@ public:
 		FGlobalTabmanager::Get()->TryInvokeTab(SearchTabName);
 	}
 
+	void OnExtendContentBrowserCommands(TSharedRef<FUICommandList> CommandList, FOnContentBrowserGetSelection GetSelectionDelegate)
+	{
+		CommandList->MapAction(FAssetSearchCommands::Get().ViewAssetSearch,
+			FExecuteAction::CreateLambda([this, GetSelectionDelegate]
+		{
+			FGlobalTabmanager::Get()->TryInvokeTab(SearchTabName);
+		})
+		);
+	}
+	
 	TSharedRef<SDockTab> HandleSpawnSettingsTab(const FSpawnTabArgs& SpawnTabArgs)
 	{
 		const TSharedRef<SDockTab> DockTab = SNew(SDockTab)
@@ -81,6 +125,8 @@ public:
 
 private:
 	TUniquePtr<FAssetSearchManager> SearchManager;
+	
+    FDelegateHandle ContentBrowserCommandExtenderDelegateHandle;
 };
 
 #undef LOCTEXT_NAMESPACE
diff --git a/Engine/Plugins/Editor/AssetSearch/Source/Public/AssetSearchCommands.h b/Engine/Plugins/Editor/AssetSearch/Source/Public/AssetSearchCommands.h
new file mode 100644
--- /dev/null
+++ b/Engine/Plugins/Editor/AssetSearch/Source/Public/AssetSearchCommands.h
@@ -0,0 +1,18 @@
+// Copyright Epic Games, Inc. All Rights Reserved.
+
+#pragma once
+
+#include "CoreMinimal.h"
+#include "Framework/Commands/Commands.h"
+
+class FAssetSearchCommands : public TCommands<FAssetSearchCommands>
+{
+public:
+	FAssetSearchCommands();
+
+	// TCommands<> interface
+	virtual void RegisterCommands() override;
+	// End of TCommands<> interface
+	
+	TSharedPtr<FUICommandInfo> ViewAssetSearch;
+};
