Asset Search - Fixing non-unity build.

diff --git a/Engine/Plugins/Editor/AssetSearch/Source/Private/AssetSearchCommands.cpp b/Engine/Plugins/Editor/AssetSearch/Source/Private/AssetSearchCommands.cpp
--- a/Engine/Plugins/Editor/AssetSearch/Source/Private/AssetSearchCommands.cpp
+++ b/Engine/Plugins/Editor/AssetSearch/Source/Private/AssetSearchCommands.cpp
@@ -1,6 +1,7 @@
 // Copyright Epic Games, Inc. All Rights Reserved.
 
 #include "AssetSearchCommands.h"
+#include "Styling/AppStyle.h" 
 
 #define LOCTEXT_NAMESPACE "AssetSearchCommands"
 
