AssetSearch - Adding a city hash for the search data when using asset tag data so we know if we need to merge it.

diff --git a/Engine/Plugins/Editor/AssetSearch/Source/Private/AssetSearchManager.cpp b/Engine/Plugins/Editor/AssetSearch/Source/Private/AssetSearchManager.cpp
--- a/Engine/Plugins/Editor/AssetSearch/Source/Private/AssetSearchManager.cpp
+++ b/Engine/Plugins/Editor/AssetSearch/Source/Private/AssetSearchManager.cpp
@@ -30,6 +30,7 @@
 #include "Materials/MaterialFunction.h"
 #include "Materials/MaterialParameterCollection.h"
 #include "Materials/MaterialInstance.h"
+#include "Hash/CityHash.h"
 
 #include "Indexers/GenericObjectIndexer.h"
 #include "Indexers/DataTableIndexer.h"
@@ -43,6 +44,7 @@
 #include "Indexers/MaterialExpressionIndexer.h"
 #include "Providers/AssetRegistrySearchProvider.h"
 
+
 #define LOCTEXT_NAMESPACE "FAssetSearchManager"
 
 static bool bForceEnableSearch = false;
@@ -169,8 +171,9 @@ private:
 	TArray<UClass*> ClassFilters;
 };
 
-const FName FAssetSearchManager::AssetSearchIndexKeyTag(TEXT("AssetSearchIndexKey"));
-const FName FAssetSearchManager::AssetSearchIndexDataTag(TEXT("AssetSearchIndexData"));
+const FName FAssetSearchManager::AssetSearchIndexVersionTag(TEXT("ASI_Version"));
+const FName FAssetSearchManager::AssetSearchIndexHashTag(TEXT("ASI_Hash"));
+const FName FAssetSearchManager::AssetSearchIndexDataTag(TEXT("ASI_Data"));
 
 FAssetSearchManager::FAssetSearchManager()
 {
@@ -380,11 +383,15 @@ void FAssetSearchManager::OnAssetAdded(const FAssetData& InAssetData)
 	static const FString UsersDeveloperPathWithSlash = FPackageName::FilenameToLongPackageName(FPaths::GameUserDeveloperDir());
 
 	const FString PackageName = InAssetData.PackageName.ToString();
+	const FString PackageMountPoint = FPackageName::GetPackageMountPoint(PackageName).ToString();
+	FString PackageFilename;
+	FPackageName::TryConvertLongPackageNameToFilename(PackageName, PackageFilename);
 
 	// Don't index things in the engine directory if we're storing data in asset tags, since we don't want to mutate the engine data.
 	if (IntermediateStorage == ESearchIntermediateStorage::AssetTagData)
 	{
-		if (PackageName.StartsWith(EngineContentPathWithSlash))
+		if (PackageName.StartsWith(EngineContentPathWithSlash) || 
+		    FPaths::IsUnderDirectory(PackageFilename, FPaths::EnginePluginsDir()))
 		{
 			return;
 		}
@@ -466,6 +473,16 @@ void FAssetSearchManager::OnAssetScanFinished()
 	});
 }
 
+uint64 FAssetSearchManager::GetTextHash(FStringView PackageRelativeExportPath) const
+{
+	if (PackageRelativeExportPath.Len() == 0)
+	{
+		return 0;
+	}
+
+	return CityHash64(reinterpret_cast<const char*>(PackageRelativeExportPath.GetData() + 1), (PackageRelativeExportPath.Len() - 1) * sizeof(TCHAR));
+}
+
 void FAssetSearchManager::HandleGetExtendedAssetRegistryTagsForSave(const UObject* Object, const ITargetPlatform* TargetPlatform, TArray<UObject::FAssetRegistryTag>& OutTags)
 {
 	if (GIsEditor && !Object->GetOutermost()->HasAnyPackageFlags(PKG_ForDiffing) && !IsRunningCookCommandlet())
@@ -480,14 +497,16 @@ void FAssetSearchManager::HandleGetExtendedAssetRegistryTagsForSave(const UObjec
 				const bool bWasIndexed = IndexAsset(InAssetData, Object, IndexedJson);
 				if (bWasIndexed && !IndexedJson.IsEmpty())
 				{
-					const FString IndexKey = GetBaseIndexKey(Object->GetClass());
+					const FString IndexVersion = GetBaseIndexKey(Object->GetClass());
+					const FString IndexHash = LexToString(GetTextHash(IndexedJson));
 
-					OutTags.Add(UObject::FAssetRegistryTag(FAssetSearchManager::AssetSearchIndexKeyTag, IndexKey, UObject::FAssetRegistryTag::TT_Hidden));
+					OutTags.Add(UObject::FAssetRegistryTag(FAssetSearchManager::AssetSearchIndexVersionTag, IndexVersion, UObject::FAssetRegistryTag::TT_Hidden));
+					OutTags.Add(UObject::FAssetRegistryTag(FAssetSearchManager::AssetSearchIndexHashTag, IndexHash, UObject::FAssetRegistryTag::TT_Hidden));
 					OutTags.Add(UObject::FAssetRegistryTag(FAssetSearchManager::AssetSearchIndexDataTag, IndexedJson, UObject::FAssetRegistryTag::TT_Hidden));
 
 					if (!IndexedJson.IsEmpty())
 					{
-						AddOrUpdateAsset(InAssetData, IndexedJson, IndexKey);
+						AddOrUpdateAsset(InAssetData, IndexedJson, IndexHash);
 					}
 				}
 			}
@@ -619,7 +638,9 @@ bool FAssetSearchManager::TryLoadIndexForAsset_Tags(const FAssetData& InAssetDat
 	FString CurrentIndexerNameAndVersion = GetBaseIndexKey(InAssetData.GetClass());
 	if (!CurrentIndexerNameAndVersion.IsEmpty())
 	{
-		const FString SavedIndexerNameAndVersion = InAssetData.GetTagValueRef<FString>(FAssetSearchManager::AssetSearchIndexKeyTag);
+		const FString SavedIndexerNameAndVersion = InAssetData.GetTagValueRef<FString>(FAssetSearchManager::AssetSearchIndexVersionTag);
+		const FString SavedIndexHash = InAssetData.GetTagValueRef<FString>(FAssetSearchManager::AssetSearchIndexHashTag);
+
 		if (!SavedIndexerNameAndVersion.IsEmpty())
 		{
 			if (SavedIndexerNameAndVersion != CurrentIndexerNameAndVersion)
@@ -628,18 +649,18 @@ bool FAssetSearchManager::TryLoadIndexForAsset_Tags(const FAssetData& InAssetDat
 				AssetNeedingReindexing.Add(InAssetData);
 			}
 
-			UpdateOperations.Enqueue([this, InAssetData, CurrentIndexerNameAndVersion]() {
+			UpdateOperations.Enqueue([this, InAssetData, SavedIndexHash]() {
 				FScopeLock ScopedLock(&SearchDatabaseCS);
-				if (!SearchDatabase.IsAssetUpToDate(InAssetData, CurrentIndexerNameAndVersion))
+				if (!SearchDatabase.IsAssetUpToDate(InAssetData, SavedIndexHash))
 				{
-					AsyncMainThreadTask([this, InAssetData, CurrentIndexerNameAndVersion]() {
+					AsyncMainThreadTask([this, InAssetData, SavedIndexHash]() {
 						const FString IndexedJson = InAssetData.GetTagValueRef<FString>(FAssetSearchManager::AssetSearchIndexDataTag);
 
 						IsAssetUpToDateCount--;
 
 						if (!IndexedJson.IsEmpty())
 						{
-							AddOrUpdateAsset(InAssetData, IndexedJson, CurrentIndexerNameAndVersion);
+							AddOrUpdateAsset(InAssetData, IndexedJson, SavedIndexHash);
 						}
 					});
 				}
@@ -1035,78 +1056,80 @@ void FAssetSearchManager::ForceIndexOnAssetsMissingIndex()
 {
 	check(IsInGameThread());
 
-	FScopedSlowTask IndexingTask(AssetNeedingReindexing.Num(), LOCTEXT("ForceIndexOnAssetsMissingIndex", "Indexing Assets"));
-	IndexingTask.MakeDialog(true);
+	{
+		FScopedSlowTask IndexingTask(AssetNeedingReindexing.Num(), LOCTEXT("ForceIndexOnAssetsMissingIndex", "Indexing Assets"));
+		IndexingTask.MakeDialog(true);
 
-	int32 RemovedCount = 0;
+		int32 RemovedCount = 0;
 
-	TArray<FAssetData> RedirectorsWithBrokenMetadata;
+		TArray<FAssetData> RedirectorsWithBrokenMetadata;
 
-	TGuardValue<bool> GuardResetTesting(GIsAutomationTesting, true);
+		TGuardValue<bool> GuardResetTesting(GIsAutomationTesting, true);
 
-	const int32 OnePercentChunk = (int32)(AssetNeedingReindexing.Num() / 100.0);
-	int32 ChunkProgress = 0;
+		const int32 OnePercentChunk = (int32)(AssetNeedingReindexing.Num() / 100.0);
+		int32 ChunkProgress = 0;
 
-	FUnloadPackageScope UnloadScope;
-	for (const FAssetData& Asset : AssetNeedingReindexing)
-	{
-		if (IndexingTask.ShouldCancel())
+		FUnloadPackageScope UnloadScope;
+		for (const FAssetData& Asset : AssetNeedingReindexing)
 		{
-			break;
-		}
+			if (IndexingTask.ShouldCancel())
+			{
+				break;
+			}
 
-		if (RemovedCount > ChunkProgress)
-		{
-			ChunkProgress += OnePercentChunk;
-			IndexingTask.EnterProgressFrame(OnePercentChunk, FText::Format(LOCTEXT("ForceIndexOnAssetsMissingIndexFormat", "Indexing Asset ({0} of {1})"), RemovedCount + 1, AssetNeedingReindexing.Num()));
-		}
+			if (RemovedCount > ChunkProgress)
+			{
+				ChunkProgress += OnePercentChunk;
+				IndexingTask.EnterProgressFrame(OnePercentChunk, FText::Format(LOCTEXT("ForceIndexOnAssetsMissingIndexFormat", "Indexing Asset ({0} of {1})"), RemovedCount + 1, AssetNeedingReindexing.Num()));
+			}
 
-		if (UObject* AssetToIndex = Asset.GetAsset())
-		{
-			// This object's metadata incorrectly labled it as something other than a redirector.  We need to resave it
-			// to stop it from appearing as something it's not.
-			if (UObjectRedirector* Redirector = Cast<UObjectRedirector>(AssetToIndex))
+			if (UObject* AssetToIndex = Asset.GetAsset())
 			{
-				RedirectorsWithBrokenMetadata.Add(Asset);
-				RemovedCount++;
-				continue;
+				// This object's metadata incorrectly labled it as something other than a redirector.  We need to resave it
+				// to stop it from appearing as something it's not.
+				if (UObjectRedirector* Redirector = Cast<UObjectRedirector>(AssetToIndex))
+				{
+					RedirectorsWithBrokenMetadata.Add(Asset);
+					RemovedCount++;
+					continue;
+				}
+
+				if (!bTryIndexAssetsOnLoad)
+				{
+					StoreIndexForAsset(AssetToIndex);
+				}
 			}
 
-			if (!bTryIndexAssetsOnLoad)
+			if (UnloadScope.GetObjectsLoaded() > 2000)
 			{
-				StoreIndexForAsset(AssetToIndex);
+				UnloadScope.TryUnload(true);
 			}
-		}
 
-		if (UnloadScope.GetObjectsLoaded() > 2000)
-		{
-			UnloadScope.TryUnload(true);
+			RemovedCount++;
 		}
 
-		RemovedCount++;
-	}
-
-	if (RedirectorsWithBrokenMetadata.Num() > 0)
-	{
-		EAppReturnType::Type ResaveRedirectors = FMessageDialog::Open(EAppMsgType::YesNo,
-			LOCTEXT("ResaveRedirectors", "We found some redirectors that didn't have the correct asset metadata identifying them as redirectors.  Would you like to resave them, so that they stop appearing as missing asset indexes?"));
-
-		if (ResaveRedirectors == EAppReturnType::Yes)
+		if (RedirectorsWithBrokenMetadata.Num() > 0)
 		{
-			TArray<UPackage*> PackagesToSave;
-			for (const FAssetData& BrokenAsset : RedirectorsWithBrokenMetadata)
+			EAppReturnType::Type ResaveRedirectors = FMessageDialog::Open(EAppMsgType::YesNo,
+				LOCTEXT("ResaveRedirectors", "We found some redirectors that didn't have the correct asset metadata identifying them as redirectors.  Would you like to resave them, so that they stop appearing as missing asset indexes?"));
+
+			if (ResaveRedirectors == EAppReturnType::Yes)
 			{
-				if (UObject* Redirector = BrokenAsset.GetAsset())
+				TArray<UPackage*> PackagesToSave;
+				for (const FAssetData& BrokenAsset : RedirectorsWithBrokenMetadata)
 				{
-					PackagesToSave.Add(Redirector->GetOutermost());
+					if (UObject* Redirector = BrokenAsset.GetAsset())
+					{
+						PackagesToSave.Add(Redirector->GetOutermost());
+					}
 				}
-			}
 
-			FEditorFileUtils::PromptForCheckoutAndSave(PackagesToSave, /*bCheckDirty*/false, /*bPromptToSave*/false);
+				FEditorFileUtils::PromptForCheckoutAndSave(PackagesToSave, /*bCheckDirty*/false, /*bPromptToSave*/false);
+			}
 		}
-	}
 
-	AssetNeedingReindexing.RemoveAtSwap(0, RemovedCount);
+		AssetNeedingReindexing.RemoveAtSwap(0, RemovedCount);
+	}
 
 	if (IntermediateStorage == ESearchIntermediateStorage::AssetTagData)
 	{
diff --git a/Engine/Plugins/Editor/AssetSearch/Source/Private/AssetSearchManager.h b/Engine/Plugins/Editor/AssetSearch/Source/Private/AssetSearchManager.h
--- a/Engine/Plugins/Editor/AssetSearch/Source/Private/AssetSearchManager.h
+++ b/Engine/Plugins/Editor/AssetSearch/Source/Private/AssetSearchManager.h
@@ -22,7 +22,8 @@ enum class ESearchIntermediateStorage : uint8;
 
 class FAssetSearchManager : public FRunnable
 {
-	static const FName AssetSearchIndexKeyTag;
+	static const FName AssetSearchIndexVersionTag;
+	static const FName AssetSearchIndexHashTag;
 	static const FName AssetSearchIndexDataTag;
 
 public:
@@ -86,6 +87,8 @@ private:
 	void AsyncMainThreadTask(TFunction<void()> Task);
 	void ProcessGameThreadTasks();
 
+	uint64 GetTextHash(FStringView PackageRelativeExportPath) const;
+
 private:
 	bool bStarted = false;
 
diff --git a/Engine/Source/Editor/Blutility/Private/AssetActionUtility.cpp b/Engine/Source/Editor/Blutility/Private/AssetActionUtility.cpp
--- a/Engine/Source/Editor/Blutility/Private/AssetActionUtility.cpp
+++ b/Engine/Source/Editor/Blutility/Private/AssetActionUtility.cpp
@@ -25,7 +25,10 @@ void UAssetActionUtility::GetAssetRegistryTags(TArray<FAssetRegistryTag>& OutTag
 	
 	FAssetActionUtilityPrototype::AddTagsFor_SupportedClasses(SupportedClasses, OutTags);
 	FAssetActionUtilityPrototype::AddTagsFor_SupportedConditions(SupportedConditions, OutTags);
-	FAssetActionUtilityPrototype::AddTagsFor_IsActionForBlueprints(IsActionForBlueprints(), OutTags);
+	if (!FUObjectThreadContext::Get().IsRoutingPostLoad)
+	{
+		FAssetActionUtilityPrototype::AddTagsFor_IsActionForBlueprints(IsActionForBlueprints(), OutTags);
+	}
 	FAssetActionUtilityPrototype::AddTagsFor_CallableFunctions(this, OutTags);
 }
 			
