Fixing a non-unity build issue.

diff --git a/Engine/Plugins/Editor/AssetSearch/Source/Private/AssetSearchModule.cpp b/Engine/Plugins/Editor/AssetSearch/Source/Private/AssetSearchModule.cpp
--- a/Engine/Plugins/Editor/AssetSearch/Source/Private/AssetSearchModule.cpp
+++ b/Engine/Plugins/Editor/AssetSearch/Source/Private/AssetSearchModule.cpp
@@ -10,6 +10,7 @@
 #include "ContentBrowserDelegates.h"
 #include "ContentBrowserModule.h"
 #include "LevelEditor.h"
+#include "Framework/Application/SlateApplication.h"
 
 #define LOCTEXT_NAMESPACE "FAssetSearchModule"
 
