Add option to disable DDC operations without changing storage strategy to asset tags.
Add option to enable integrity checks instead of doing them by default - they can take several minutes in a very large (e.g. 50gb) db, and prolong editor shutdown.
Fix very slow search in FAssetRegistrySearchProvider caused by shuffling objects in TArray which can prolong editor shutdown.

diff --git a/Engine/Plugins/Editor/AssetSearch/Source/Private/AssetSearchDatabase.cpp b/Engine/Plugins/Editor/AssetSearch/Source/Private/AssetSearchDatabase.cpp
--- a/Engine/Plugins/Editor/AssetSearch/Source/Private/AssetSearchDatabase.cpp
+++ b/Engine/Plugins/Editor/AssetSearch/Source/Private/AssetSearchDatabase.cpp
@@ -474,6 +474,7 @@ bool FAssetSearchDatabase::Open(const FString& InSessionPath)
 
 bool FAssetSearchDatabase::Open(const FString& InSessionPath, const ESQLiteDatabaseOpenMode InOpenMode)
 {
+	TRACE_CPUPROFILER_EVENT_SCOPE_STR("FAssetSearchDatabase::Open");
 	SessionPath = InSessionPath;
 
 	if (Database->IsValid())
@@ -487,7 +488,7 @@ bool FAssetSearchDatabase::Open(const FString& InSessionPath, const ESQLiteDatab
 		return false;
 	}
 
-	if (!Database->PerformQuickIntegrityCheck())
+	if (bEnableIntegrityChecks && !Database->PerformQuickIntegrityCheck())
 	{
 		UE_LOG(LogAssetSearch, Error, TEXT("Database failed integrity check, deleting."));
 
@@ -668,7 +669,7 @@ bool FAssetSearchDatabase::Open(const FString& InSessionPath, const ESQLiteDatab
 		return false;
 	}
 
-	if (!Database->PerformQuickIntegrityCheck())
+	if (bEnableIntegrityChecks && !Database->PerformQuickIntegrityCheck())
 	{
 		UE_LOG(LogAssetSearch, Error, TEXT("Database failed integrity check, deleting."));
 
diff --git a/Engine/Plugins/Editor/AssetSearch/Source/Private/AssetSearchDatabase.h b/Engine/Plugins/Editor/AssetSearch/Source/Private/AssetSearchDatabase.h
--- a/Engine/Plugins/Editor/AssetSearch/Source/Private/AssetSearchDatabase.h
+++ b/Engine/Plugins/Editor/AssetSearch/Source/Private/AssetSearchDatabase.h
@@ -99,6 +99,8 @@ public:
 	 * Remove asset from the the search database that are not in this set of assets.
 	 */
 	void RemoveAssetsNotInThisSet(const TArray<FAssetData>& InAssets);
+	
+	bool bEnableIntegrityChecks = false;
 
 private:
 	void LogLastError() const;
diff --git a/Engine/Plugins/Editor/AssetSearch/Source/Private/AssetSearchManager.cpp b/Engine/Plugins/Editor/AssetSearch/Source/Private/AssetSearchManager.cpp
--- a/Engine/Plugins/Editor/AssetSearch/Source/Private/AssetSearchManager.cpp
+++ b/Engine/Plugins/Editor/AssetSearch/Source/Private/AssetSearchManager.cpp
@@ -234,6 +234,8 @@ void FAssetSearchManager::Start()
 
 	TickerHandle = FTSTicker::GetCoreTicker().AddTicker(FTickerDelegate::CreateRaw(this, &FAssetSearchManager::Tick_GameThread), 0);
 
+	SearchDatabase.bEnableIntegrityChecks = GetDefault<USearchUserSettings>()->bEnableIntegrityChecks;
+	FileInfoDatabase.bEnableIntegrityChecks = GetDefault<USearchUserSettings>()->bEnableIntegrityChecks;
 	RunThread = true;
 	DatabaseThread = FRunnableThread::Create(this, TEXT("UniversalSearch"), 0, TPri_BelowNormal);
 
@@ -685,6 +687,12 @@ bool FAssetSearchManager::TryLoadIndexForAsset_Tags(const FAssetData& InAssetDat
 bool FAssetSearchManager::TryLoadIndexForAsset_DDC(const FAssetData& InAssetData)
 {
 	check(IntermediateStorage == ESearchIntermediateStorage::DerivedDataCache);
+	
+	bool bAllowFetch = !GetDefault<USearchProjectSettings>()->bDisableDDC;
+	if (!bAllowFetch)
+	{
+		return false;
+	}
 
 	const bool bSuccess = AsyncGetDerivedDataKey(InAssetData, [this, InAssetData](bool bSuccess, FString InDDCKey) {
 		if (!bSuccess)
@@ -902,9 +910,13 @@ void FAssetSearchManager::StoreIndexForAsset_DDC(const UObject* InAsset)
 			AsyncMainThreadTask([this, InAssetData, IndexedJson, InDDCKey]() {
 				check(IsInGameThread());
 
-				FTCHARToUTF8 IndexedJsonUTF8(*IndexedJson);
-				TArrayView<const uint8> IndexedJsonUTF8View((const uint8*)IndexedJsonUTF8.Get(), IndexedJsonUTF8.Length() * sizeof(UTF8CHAR));
-				GetDerivedDataCacheRef().Put(*InDDCKey, IndexedJsonUTF8View, InAssetData.GetObjectPathString(), false);
+				bool bAllowPut = !GetDefault<USearchProjectSettings>()->bDisableDDC;
+				if (bAllowPut)
+				{
+					FTCHARToUTF8 IndexedJsonUTF8(*IndexedJson);
+					TArrayView<const uint8> IndexedJsonUTF8View((const uint8*)IndexedJsonUTF8.Get(), IndexedJsonUTF8.Length() * sizeof(UTF8CHAR));
+					GetDerivedDataCacheRef().Put(*InDDCKey, IndexedJsonUTF8View, InAssetData.GetObjectPathString(), false);
+				}
 
 				AddOrUpdateAsset(InAssetData, IndexedJson, InDDCKey);
 			});
diff --git a/Engine/Plugins/Editor/AssetSearch/Source/Private/AssetSearchModule.cpp b/Engine/Plugins/Editor/AssetSearch/Source/Private/AssetSearchModule.cpp
--- a/Engine/Plugins/Editor/AssetSearch/Source/Private/AssetSearchModule.cpp
+++ b/Engine/Plugins/Editor/AssetSearch/Source/Private/AssetSearchModule.cpp
@@ -72,6 +72,8 @@ public:
 			TSharedRef<FUICommandList> CommandList = LevelEditorModule.GetGlobalLevelEditorActions();
 			CommandList->UnmapAction(FAssetSearchCommands::Get().ViewAssetSearch);
 		}
+		
+		SearchManager.Reset();
 	}
 
 	void ExecuteOpenObjectBrowser()
diff --git a/Engine/Plugins/Editor/AssetSearch/Source/Private/FileInfoDatabase.h b/Engine/Plugins/Editor/AssetSearch/Source/Private/FileInfoDatabase.h
--- a/Engine/Plugins/Editor/AssetSearch/Source/Private/FileInfoDatabase.h
+++ b/Engine/Plugins/Editor/AssetSearch/Source/Private/FileInfoDatabase.h
@@ -80,6 +80,8 @@ public:
 	void UpdateFileHash(const FAssetData& InAssetData, FAssetFileInfo& OutFileInfo);
 
 	TMap<FName, FAssetFileInfo> GetAllFileInfos();
+	
+	bool bEnableIntegrityChecks = false;
 
 private:
 	void LogLastError() const;
diff --git a/Engine/Plugins/Editor/AssetSearch/Source/Private/Providers/AssetRegistrySearchProvider.cpp b/Engine/Plugins/Editor/AssetSearch/Source/Private/Providers/AssetRegistrySearchProvider.cpp
--- a/Engine/Plugins/Editor/AssetSearch/Source/Private/Providers/AssetRegistrySearchProvider.cpp
+++ b/Engine/Plugins/Editor/AssetSearch/Source/Private/Providers/AssetRegistrySearchProvider.cpp
@@ -345,33 +345,29 @@ void FAssetRegistrySearchProvider::Search(FSearchQueryPtr SearchQuery)
 	Registry.GetAllAssets(Assets);
 
 	Async(EAsyncExecution::LargeThreadPool, [Assets = MoveTemp(Assets), SearchQuery]() mutable {
+		TRACE_CPUPROFILER_EVENT_SCOPE_STR("FAssetRegistrySearchProvider Async Search");
 		FTextFilterExpressionEvaluator TextFilterExpressionEvaluator(ETextFilterExpressionEvaluatorMode::Complex);
 		TextFilterExpressionEvaluator.SetFilterText(FText::FromString(SearchQuery->QueryText));
 		FFrontendFilter_TextFilterExpressionContext_HackCopy TextFilterExpressionContext;
 
+		TArray<FSearchRecord> SearchResults;
 		for (auto AssetIter = Assets.CreateIterator(); AssetIter; ++AssetIter)
 		{
 			const FAssetData& Asset = *AssetIter;
 			TextFilterExpressionContext.SetAsset(&Asset);
-			if (!TextFilterExpressionEvaluator.TestTextFilter(TextFilterExpressionContext))
+			if (TextFilterExpressionEvaluator.TestTextFilter(TextFilterExpressionContext))
 			{
-				AssetIter.RemoveCurrent();
-			}
-			TextFilterExpressionContext.ClearAsset();
-		}
+				FSearchRecord Record;
+				Record.AssetPath = Asset.GetObjectPathString();
+				Record.AssetName = Asset.AssetName.ToString();
+				Record.AssetClass = Asset.AssetClassPath;
 
-		TArray<FSearchRecord> SearchResults;
-		for (const FAssetData& Asset : Assets)
-		{
-			FSearchRecord Record;
-			Record.AssetPath = Asset.GetObjectPathString();
-			Record.AssetName = Asset.AssetName.ToString();
-			Record.AssetClass = Asset.AssetClassPath;
-
-			const float WorstCase = Record.AssetName.Len() + SearchQuery->QueryText.Len();
-			Record.Score = -50.0f * (1.0f - (Algo::LevenshteinDistance(Record.AssetName.ToLower(), SearchQuery->QueryText.ToLower()) / WorstCase));
+				const float WorstCase = Record.AssetName.Len() + SearchQuery->QueryText.Len();
+				Record.Score = -50.0f * (1.0f - (Algo::LevenshteinDistance(Record.AssetName.ToLower(), SearchQuery->QueryText.ToLower()) / WorstCase));
 
-			SearchResults.Add(Record);
+				SearchResults.Add(Record);
+			}
+			TextFilterExpressionContext.ClearAsset();
 		}
 
 		Async(EAsyncExecution::TaskGraphMainThread, [SearchQuery, SearchResults = MoveTemp(SearchResults)]() mutable {
diff --git a/Engine/Plugins/Editor/AssetSearch/Source/Private/Settings/SearchProjectSettings.h b/Engine/Plugins/Editor/AssetSearch/Source/Private/Settings/SearchProjectSettings.h
--- a/Engine/Plugins/Editor/AssetSearch/Source/Private/Settings/SearchProjectSettings.h
+++ b/Engine/Plugins/Editor/AssetSearch/Source/Private/Settings/SearchProjectSettings.h
@@ -30,4 +30,8 @@ public:
 
 	UPROPERTY(config, EditAnywhere, Category=General)
 	TArray<FDirectoryPath> IgnoredPaths;
+	
+	/** Disable put/fetch operations with the DDC, without changing storage location */
+	UPROPERTY(config, EditAnywhere, Category=General)
+	bool bDisableDDC;
 };
diff --git a/Engine/Plugins/Editor/AssetSearch/Source/Private/Settings/SearchUserSettings.h b/Engine/Plugins/Editor/AssetSearch/Source/Private/Settings/SearchUserSettings.h
--- a/Engine/Plugins/Editor/AssetSearch/Source/Private/Settings/SearchUserSettings.h
+++ b/Engine/Plugins/Editor/AssetSearch/Source/Private/Settings/SearchUserSettings.h
@@ -34,6 +34,10 @@ public:
 	/** Enable this to begin using search. */
 	UPROPERTY(config, EditAnywhere, Category = General)
 	bool bEnableSearch;
+	
+	/** Enable sqlite integrity checks on database load/setup. Slow. */
+	UPROPERTY(config, EditAnywhere, Category=General)
+	bool bEnableIntegrityChecks;
 
 	UPROPERTY(config, EditAnywhere, Category=General)
 	TArray<FDirectoryPath> IgnoredPaths;
