Asset Indexing: enable indexing of ScriptStructs (e.g. InstancedStructs) inside properties flagged as AssetRegistrySearchable.

#rnx

diff --git a/Engine/Plugins/Editor/AssetSearch/Source/Private/Utility/IndexerUtilities.cpp b/Engine/Plugins/Editor/AssetSearch/Source/Private/Utility/IndexerUtilities.cpp
--- a/Engine/Plugins/Editor/AssetSearch/Source/Private/Utility/IndexerUtilities.cpp
+++ b/Engine/Plugins/Editor/AssetSearch/Source/Private/Utility/IndexerUtilities.cpp
@@ -149,6 +149,17 @@ void FIndexerUtilities::IterateIndexableProperties(const UStruct* InStruct, cons
 				It.SkipRecursiveProperty();
 				continue;
 			}
+			else if (UScriptStruct* ScriptStruct = StructProperty->Struct.Get())
+			{
+				// We don't want to record all ScriptStructs, only those explicitly flagged.
+				TArray<const FProperty*> PropertyChain;
+				It.GetPropertyChain(PropertyChain);
+				const bool bIsSearchable = PropertyChain.ContainsByPredicate([](const FProperty* Property) { return Property->HasAnyPropertyFlags(CPF_AssetRegistrySearchable); });
+				if (bIsSearchable)
+				{
+					ScriptStruct->ExportText(Text, ValuePtr, nullptr, nullptr, PPF_None, nullptr);
+				}
+			}
 			//else if (StructProperty->Struct == FGuid::StaticStruct())
 			//{
 			//	const FGuid* Guid = static_cast<const FGuid*>(ValuePtr);
