NOTE: OnGetExtraObjectTagsWithContext delegates must now be thread safe.
* Converting OnGetExtraObjectTagsWithContext to be a threadsafe multicast delegate, fixing TMulticastScriptDelegate to no longer require write access when being serialized, and adding locks to AssetSearchManager to ensure that it can be safely accessed from multiple threads during asset tag serialization.

diff --git a/Engine/Plugins/Editor/AssetSearch/Source/Private/AssetSearchManager.cpp b/Engine/Plugins/Editor/AssetSearch/Source/Private/AssetSearchManager.cpp
--- a/Engine/Plugins/Editor/AssetSearch/Source/Private/AssetSearchManager.cpp
+++ b/Engine/Plugins/Editor/AssetSearch/Source/Private/AssetSearchManager.cpp
@@ -365,21 +365,21 @@ FSearchStats FAssetSearchManager::GetStats() const
 
 void FAssetSearchManager::RegisterAssetIndexer(const UClass* AssetClass, TUniquePtr<IAssetIndexer>&& Indexer)
 {
-	check(IsInGameThread());
+	FScopeLock Lock(&AssetSearchManagerCS);
 
 	Indexers.Add(AssetClass->GetFName(), MoveTemp(Indexer));
 }
 
 void FAssetSearchManager::RegisterSearchProvider(FName SearchProviderName, TUniquePtr<ISearchProvider>&& InSearchProvider)
 {
-	check(IsInGameThread());
+	FScopeLock Lock(&AssetSearchManagerCS);
 
 	SearchProviders.Add(SearchProviderName, MoveTemp(InSearchProvider));
 }
 
 void FAssetSearchManager::OnAssetAdded(const FAssetData& InAssetData)
 {
-	check(IsInGameThread());
+	FScopeLock Lock(&AssetSearchManagerCS);
 
 	static const FString EngineContentPathWithSlash = FPackageName::FilenameToLongPackageName(FPaths::EngineContentDir());
 	static const FString DeveloperPathWithSlash = FPackageName::FilenameToLongPackageName(FPaths::GameDevelopersDir());
@@ -442,7 +442,7 @@ void FAssetSearchManager::OnAssetAdded(const FAssetData& InAssetData)
 
 void FAssetSearchManager::OnAssetRemoved(const FAssetData& InAssetData)
 {
-	check(IsInGameThread());
+	FScopeLock Lock(&AssetSearchManagerCS);
 
 	FAssetOperation Operation;
 	Operation.Asset = InAssetData;
@@ -452,7 +452,7 @@ void FAssetSearchManager::OnAssetRemoved(const FAssetData& InAssetData)
 
 void FAssetSearchManager::OnAssetScanFinished()
 {
-	check(IsInGameThread());
+	FScopeLock Lock(&AssetSearchManagerCS);
 
 	TArray<FAssetData> AllAssets;
 	IAssetRegistry& AssetRegistry = FModuleManager::LoadModuleChecked<FAssetRegistryModule>("AssetRegistry").Get();
@@ -501,6 +501,8 @@ void FAssetSearchManager::HandleOnGetExtraObjectTags(FAssetRegistryTagsContext C
 		{
 			if (!TargetPlatform || TargetPlatform->HasEditorOnlyData())
 			{
+				FScopeLock Lock(&AssetSearchManagerCS);
+
 				FAssetData InAssetData(Object, FAssetData::ECreationFlags::SkipAssetRegistryTagsGathering);
 
 				FString IndexedJson;
@@ -526,7 +528,6 @@ void FAssetSearchManager::HandleOnGetExtraObjectTags(FAssetRegistryTagsContext C
 
 void FAssetSearchManager::HandlePackageSaved(const FString& PackageFilename, UPackage* Package, FObjectPostSaveContext ObjectSaveContext)
 {
-	check(IsInGameThread());
 	check(IntermediateStorage == ESearchIntermediateStorage::DerivedDataCache);
 
 	// Only execute if this is a user save
@@ -537,6 +538,7 @@ void FAssetSearchManager::HandlePackageSaved(const FString& PackageFilename, UPa
 
 	if (GIsEditor && !IsRunningCommandlet())
 	{
+		FScopeLock Lock(&AssetSearchManagerCS);
 		TArray<UObject*> Objects;
 		const bool bIncludeNestedObjects = false;
 		GetObjectsWithPackage(Package, Objects, bIncludeNestedObjects);
@@ -549,8 +551,6 @@ void FAssetSearchManager::HandlePackageSaved(const FString& PackageFilename, UPa
 
 void FAssetSearchManager::OnAssetLoaded(UObject* InObject)
 {
-	check(IsInGameThread());
-
 	if (bTryIndexAssetsOnLoad)
 	{
 		switch (IntermediateStorage)
@@ -568,7 +568,7 @@ void FAssetSearchManager::OnAssetLoaded(UObject* InObject)
 
 bool FAssetSearchManager::RequestIndexAsset_DDC(const UObject* InAsset)
 {
-	check(IsInGameThread());
+	FScopeLock Lock(&AssetSearchManagerCS);
 	check(IntermediateStorage == ESearchIntermediateStorage::DerivedDataCache);
 
 	if (GEditor == nullptr || GEditor->IsAutosaving())
@@ -640,7 +640,7 @@ bool FAssetSearchManager::TryLoadIndexForAsset(const FAssetData& InAssetData)
 
 bool FAssetSearchManager::TryLoadIndexForAsset_Tags(const FAssetData& InAssetData)
 {
-	check(IsInGameThread());
+	FScopeLock Lock(&AssetSearchManagerCS);
 	check(IntermediateStorage == ESearchIntermediateStorage::AssetTagData);
 
 	IsAssetUpToDateCount++;
@@ -694,6 +694,7 @@ bool FAssetSearchManager::TryLoadIndexForAsset_Tags(const FAssetData& InAssetDat
 
 bool FAssetSearchManager::TryLoadIndexForAsset_DDC(const FAssetData& InAssetData)
 {
+	FScopeLock Lock(&AssetSearchManagerCS);
 	check(IntermediateStorage == ESearchIntermediateStorage::DerivedDataCache);
 	
 	bool bAllowFetch = !GetDefault<USearchProjectSettings>()->bDisableDDC;
@@ -730,6 +731,7 @@ bool FAssetSearchManager::TryLoadIndexForAsset_DDC(const FAssetData& InAssetData
 
 void FAssetSearchManager::AsyncRequestDownload(const FAssetData& InAssetData, const FString& InDDCKey)
 {
+	FScopeLock Lock(&AssetSearchManagerCS);
 	DownloadQueueCount++;
 
 	FAssetDDCRequest DDCRequest;
@@ -740,7 +742,7 @@ void FAssetSearchManager::AsyncRequestDownload(const FAssetData& InAssetData, co
 
 bool FAssetSearchManager::AsyncGetDerivedDataKey(const FAssetData& InAssetData, TFunction<void(bool, FString)> DDCKeyCallback)
 {
-	check(IsInGameThread());
+	FScopeLock Lock(&AssetSearchManagerCS);
 	check(IntermediateStorage == ESearchIntermediateStorage::DerivedDataCache);
 
 	const FString AssetPath = InAssetData.PackagePath.ToString();
@@ -863,7 +865,7 @@ FString FAssetSearchManager::GetIndexerVersion(const UClass* InAssetClass) const
 
 bool FAssetSearchManager::IndexAsset(const FAssetData& InAssetData, const UObject* InAsset, FString& OutIndexedJson) const
 {
-	check(IsInGameThread());
+	FScopeLock Lock(&AssetSearchManagerCS);
 
 	if (IsAssetIndexable(InAsset) && HasIndexerForClass(InAsset->GetClass()))
 	{
@@ -889,7 +891,7 @@ void FAssetSearchManager::StoreIndexForAsset(const UObject* InAsset)
 
 void FAssetSearchManager::StoreIndexForAsset_Tags(const UObject* InAsset)
 {
-	check(IsInGameThread());
+	FScopeLock Lock(&AssetSearchManagerCS);
 	check(IntermediateStorage == ESearchIntermediateStorage::AssetTagData);
 
 	// StoreIndexForAsset doesn't really do much of anything when it's stored in tag data, all we can really
@@ -900,7 +902,7 @@ void FAssetSearchManager::StoreIndexForAsset_Tags(const UObject* InAsset)
 
 void FAssetSearchManager::StoreIndexForAsset_DDC(const UObject* InAsset)
 {
-	check(IsInGameThread());
+	FScopeLock Lock(&AssetSearchManagerCS);
 	check(IntermediateStorage == ESearchIntermediateStorage::DerivedDataCache);
 
 	FString IndexedJson;
@@ -942,7 +944,7 @@ void FAssetSearchManager::LoadDDCContentIntoDatabase(const FAssetData& InAsset,
 
 void FAssetSearchManager::AddOrUpdateAsset(const FAssetData& InAssetData, const FString& IndexedJson, const FString& DerivedDataKey)
 {
-	check(IsInGameThread());
+	FScopeLock Lock(&AssetSearchManagerCS);
 
 	PendingDatabaseUpdates++;
 	UpdateOperations.Enqueue([this, InAssetData, IndexedJson, DerivedDataKey]() {
@@ -956,6 +958,7 @@ bool FAssetSearchManager::Tick_GameThread(float DeltaTime)
 {
 	QUICK_SCOPE_CYCLE_COUNTER(STAT_FAssetSearchManager_Tick);
 
+	FScopeLock Lock(&AssetSearchManagerCS);
 	check(IsInGameThread());
 
 	UpdateScanningAssets();
@@ -1074,7 +1077,7 @@ void FAssetSearchManager::Tick_DatabaseOperationThread()
 
 void FAssetSearchManager::ForceIndexOnAssetsMissingIndex()
 {
-	check(IsInGameThread());
+	FScopeLock Lock(&AssetSearchManagerCS);
 
 	{
 		FScopedSlowTask IndexingTask(AssetNeedingReindexing.Num(), LOCTEXT("ForceIndexOnAssetsMissingIndex", "Indexing Assets"));
@@ -1166,6 +1169,7 @@ void FAssetSearchManager::ForceIndexOnAssetsMissingIndex()
 void FAssetSearchManager::Search(FSearchQueryPtr SearchQuery)
 {
 	check(IsInGameThread());
+	FScopeLock Lock(&AssetSearchManagerCS);
 
 	if (FStudioTelemetry::IsAvailable())
 	{
@@ -1205,11 +1209,13 @@ void FAssetSearchManager::Search(FSearchQueryPtr SearchQuery)
 
 void FAssetSearchManager::AsyncMainThreadTask(TFunction<void()> Task)
 {
+	FScopeLock Lock(&AssetSearchManagerCS);
 	GT_Tasks.Enqueue(Task);
 }
 
 void FAssetSearchManager::ProcessGameThreadTasks()
 {
+	FScopeLock Lock(&AssetSearchManagerCS);
 	if (!GT_Tasks.IsEmpty())
 	{
 		if (GIsSavingPackage)
diff --git a/Engine/Plugins/Editor/AssetSearch/Source/Private/AssetSearchManager.h b/Engine/Plugins/Editor/AssetSearch/Source/Private/AssetSearchManager.h
--- a/Engine/Plugins/Editor/AssetSearch/Source/Private/AssetSearchManager.h
+++ b/Engine/Plugins/Editor/AssetSearch/Source/Private/AssetSearchManager.h
@@ -93,6 +93,7 @@ private:
 	bool bStarted = false;
 
 private:
+	mutable FCriticalSection AssetSearchManagerCS;
 	FFileInfoDatabase FileInfoDatabase;
 	FCriticalSection FileInfoDatabaseCS;
 	FAssetSearchDatabase SearchDatabase;
diff --git a/Engine/Source/Runtime/Core/Public/UObject/ScriptDelegates.h b/Engine/Source/Runtime/Core/Public/UObject/ScriptDelegates.h
--- a/Engine/Source/Runtime/Core/Public/UObject/ScriptDelegates.h
+++ b/Engine/Source/Runtime/Core/Public/UObject/ScriptDelegates.h
@@ -837,31 +837,31 @@ public:
 	/** Multi-cast delegate serialization */
 	friend FArchive& operator<<( FArchive& Ar, TMulticastScriptDelegate& D )
 	{
-		// Special case to avoid taking a lock on empty script delegate.
-		// This is required for avoiding asserts on EmptyDelegate serialization.
-		if (Ar.IsSaving() && D.InvocationList.Num() == 0)
+		// Warning: This order of if/else is necessary because there are archivers (eg FArchiveGatherExternalActorRefs)
+		// that do not report as either IsLoading or IsSaving but should be treated as IsSaving in this case
+		if (Ar.IsLoading())
 		{
-			FReadAccessScope ReadScope = D.GetReadAccessScope();
-
+			FWriteAccessScope WriteScope = D.GetWriteAccessScope();
 			Ar << D.InvocationList;
+			// After loading the delegate, clean up the list to make sure there are no bad object references
+			D.CompactInvocationList();
 		}
 		else
 		{
-			FWriteAccessScope WriteScope = D.GetWriteAccessScope();
-
-			if( Ar.IsSaving() )
-			{
-				// When saving the delegate, clean up the list to make sure there are no bad object references
-				D.CompactInvocationList();
-			}
-
-			Ar << D.InvocationList;
+			FReadAccessScope ReadScope = D.GetReadAccessScope();
+			// When saving the delegate, clean up the list to make sure there are no bad object references
+			// Don't do this in place because we don't want to require a write lock
+			typedef TArray<UnicastDelegateType, TInlineAllocator<4>> FInlineInvocationList;
+			FInlineInvocationList CompactedList;
 
-			if( Ar.IsLoading() )
+			for (const UnicastDelegateType& Delegate : D.InvocationList)
 			{
-				// After loading the delegate, clean up the list to make sure there are no bad object references
-				D.CompactInvocationList();
+				if (!Delegate.IsCompactable())
+				{
+					CompactedList.Add(Delegate);
+				}
 			}
+			Ar << CompactedList;
 		}
 
 		return Ar;
diff --git a/Engine/Source/Runtime/CoreUObject/Public/UObject/Object.h b/Engine/Source/Runtime/CoreUObject/Public/UObject/Object.h
--- a/Engine/Source/Runtime/CoreUObject/Public/UObject/Object.h
+++ b/Engine/Source/Runtime/CoreUObject/Public/UObject/Object.h
@@ -858,14 +858,16 @@ public:
 			: Name(InName), Value(MoveTemp(InValue)), Type(InType), DisplayFlags(InDisplayFlags) {}
 
 #if WITH_EDITOR
-		/** Event for listeners who want to add tags to some UObjects' GetAssetRegistryTags. */
-		DECLARE_MULTICAST_DELEGATE_OneParam(FOnGetObjectAssetRegistryTagsWithContext, FAssetRegistryTagsContext);
+		/** Event for listeners who want to add tags to some UObjects' GetAssetRegistryTags. 
+		  * Listeners should be threadsafe as this event may be broadcast on a background thread during concurrent save
+		  **/
+		DECLARE_TS_MULTICAST_DELEGATE_OneParam(FOnGetObjectAssetRegistryTagsWithContext, FAssetRegistryTagsContext);
 		COREUOBJECT_API static FOnGetObjectAssetRegistryTagsWithContext OnGetExtraObjectTagsWithContext;
 
-		DECLARE_MULTICAST_DELEGATE_TwoParams(FOnGetObjectAssetRegistryTags, const UObject* /*Object*/, TArray<FAssetRegistryTag>& /*InOutTags*/);
+		DECLARE_TS_MULTICAST_DELEGATE_TwoParams(FOnGetObjectAssetRegistryTags, const UObject* /*Object*/, TArray<FAssetRegistryTag>& /*InOutTags*/);
 		UE_DEPRECATED(5.4, "Subscribe to OnGetExtraObjectTagsWithContext instead")
 		COREUOBJECT_API static FOnGetObjectAssetRegistryTags OnGetExtraObjectTags;
-		DECLARE_MULTICAST_DELEGATE_ThreeParams(FOnGetExtendedAssetRegistryTagsForSave, const UObject* /*Object*/, const ITargetPlatform* TargetPlatform, TArray<FAssetRegistryTag>& /*InOutTags*/);
+		DECLARE_TS_MULTICAST_DELEGATE_ThreeParams(FOnGetExtendedAssetRegistryTagsForSave, const UObject* /*Object*/, const ITargetPlatform* TargetPlatform, TArray<FAssetRegistryTag>& /*InOutTags*/);
 		UE_DEPRECATED(5.4, "Subscribe to OnGetExtraObjectTagsWithContext instead, and early exit if !Context.IsSaving")
 		COREUOBJECT_API static FOnGetExtendedAssetRegistryTagsForSave OnGetExtendedAssetRegistryTagsForSave;
 #endif // WITH_EDITOR
