[AssetSearch]: Update FAssetRegistrySearchProvider to only include assets on the disk in order to prevent it from blocking main thread
#rnx

diff --git a/Engine/Plugins/Editor/AssetSearch/Source/Private/Providers/AssetRegistrySearchProvider.cpp b/Engine/Plugins/Editor/AssetSearch/Source/Private/Providers/AssetRegistrySearchProvider.cpp
--- a/Engine/Plugins/Editor/AssetSearch/Source/Private/Providers/AssetRegistrySearchProvider.cpp
+++ b/Engine/Plugins/Editor/AssetSearch/Source/Private/Providers/AssetRegistrySearchProvider.cpp
@@ -338,14 +338,15 @@ private:
 
 void FAssetRegistrySearchProvider::Search(FSearchQueryPtr SearchQuery)
 {
-	IAssetRegistry& Registry = FAssetRegistryModule::GetRegistry();
-
-	// Start by gathering all assets.
-	TArray<FAssetData> Assets;
-	Registry.GetAllAssets(Assets);
-
-	Async(EAsyncExecution::LargeThreadPool, [Assets = MoveTemp(Assets), SearchQuery]() mutable {
+	Async(EAsyncExecution::LargeThreadPool, [SearchQuery]() mutable {
 		TRACE_CPUPROFILER_EVENT_SCOPE_STR("FAssetRegistrySearchProvider Async Search");
+		
+		const IAssetRegistry& Registry = FAssetRegistryModule::GetRegistry();
+			
+		// Start by gathering all assets.
+		TArray<FAssetData> Assets;
+		Registry.GetAllAssets(Assets, true /* bIncludeOnlyOnDiskAssets */ );
+		
 		FTextFilterExpressionEvaluator TextFilterExpressionEvaluator(ETextFilterExpressionEvaluatorMode::Complex);
 		TextFilterExpressionEvaluator.SetFilterText(FText::FromString(SearchQuery->QueryText));
 		FFrontendFilter_TextFilterExpressionContext_HackCopy TextFilterExpressionContext;
