Fixed printf mismatch.

diff --git a/Engine/Plugins/Editor/AssetSearch/Source/Private/FileInfoDatabase.cpp b/Engine/Plugins/Editor/AssetSearch/Source/Private/FileInfoDatabase.cpp
--- a/Engine/Plugins/Editor/AssetSearch/Source/Private/FileInfoDatabase.cpp
+++ b/Engine/Plugins/Editor/AssetSearch/Source/Private/FileInfoDatabase.cpp
@@ -404,7 +404,7 @@ FString FFileInfoDatabase::GetLastError() const
 
 void FFileInfoDatabase::LogLastError() const
 {
-	UE_LOG(LogFileInfo, Error, TEXT("Database Error: %s"), *SessionPath, *GetLastError());
+	UE_LOG(LogFileInfo, Error, TEXT("Database '%s' Error: %s"), *SessionPath, *GetLastError());
 }
 
 bool FFileInfoDatabase::AddOrUpdateFileInfo(const FAssetData& InAssetData, FAssetFileInfo& OutFileInfo)
