Asset Search Manager: Do not reload all packages after each AT.
This halves AT runtime in editor for a large licensee project.
#rnx

diff --git a/Engine/Plugins/Editor/AssetSearch/Source/Private/AssetSearchManager.cpp b/Engine/Plugins/Editor/AssetSearch/Source/Private/AssetSearchManager.cpp
--- a/Engine/Plugins/Editor/AssetSearch/Source/Private/AssetSearchManager.cpp
+++ b/Engine/Plugins/Editor/AssetSearch/Source/Private/AssetSearchManager.cpp
@@ -246,6 +246,7 @@ void FAssetSearchManager::Start()
 	}
 }
 
+// Called from within AssetSearchManagerCS lock
 void FAssetSearchManager::UpdateScanningAssets()
 {
 	bool bTargetState = GetDefault<USearchUserSettings>()->bEnableSearch;
@@ -275,25 +276,22 @@ void FAssetSearchManager::UpdateScanningAssets()
 	}
 }
 
+// Called from within AssetSearchManagerCS lock
 void FAssetSearchManager::StartScanningAssets()
 {
 	IAssetRegistry& AssetRegistry = FModuleManager::LoadModuleChecked<FAssetRegistryModule>("AssetRegistry").Get();
 	if (AssetRegistry.IsLoadingAssets())
 	{
-		AssetRegistry.OnFilesLoaded().AddRaw(this, &FAssetSearchManager::OnAssetScanFinished);
+		AssetRegistry.OnFilesLoaded().AddRaw(this, &FAssetSearchManager::StartScanningAssetsOnAssetScanFinished);
+	}
+	else if (!bFinishedOnce)
+	{
+		StartScanningAssetsOnAssetScanFinishedInternal();
 	}
 	else
 	{
 		AssetRegistry.OnAssetAdded().AddRaw(this, &FAssetSearchManager::OnAssetAdded);
 		AssetRegistry.OnAssetRemoved().AddRaw(this, &FAssetSearchManager::OnAssetRemoved);
-
-		TArray<FAssetData> TempAssetData;
-		AssetRegistry.GetAllAssets(TempAssetData, true);
-
-		for (const FAssetData& Data : TempAssetData)
-		{
-			OnAssetAdded(Data);
-		}
 	}
 }
 
@@ -450,30 +448,40 @@ void FAssetSearchManager::OnAssetRemoved(const FAssetData& InAssetData)
 	ProcessAssetQueue.Add(Operation);
 }
 
-void FAssetSearchManager::OnAssetScanFinished()
+void FAssetSearchManager::StartScanningAssetsOnAssetScanFinished()
 {
 	FScopeLock Lock(&AssetSearchManagerCS);
 
-	TArray<FAssetData> AllAssets;
+	TArray<FAssetData> AllAssets = StartScanningAssetsOnAssetScanFinishedInternal();
 	IAssetRegistry& AssetRegistry = FModuleManager::LoadModuleChecked<FAssetRegistryModule>("AssetRegistry").Get();
 	AssetRegistry.OnFilesLoaded().RemoveAll(this);
+	
+	PendingDatabaseUpdates++;
+	UpdateOperations.Enqueue([this, AssetsAvailable = MoveTemp(AllAssets)]() mutable {
+		FScopeLock ScopedLock(&SearchDatabaseCS);
+		//UE_LOG(LogAssetSearch, Log, TEXT(""));
+		SearchDatabase.RemoveAssetsNotInThisSet(AssetsAvailable);
+		PendingDatabaseUpdates--;
+	});
+}
+
+// Called from within AssetSearchManagerCS lock
+TArray<FAssetData> FAssetSearchManager::StartScanningAssetsOnAssetScanFinishedInternal()
+{
+	IAssetRegistry& AssetRegistry = FModuleManager::LoadModuleChecked<FAssetRegistryModule>("AssetRegistry").Get();
 	AssetRegistry.OnAssetAdded().AddRaw(this, &FAssetSearchManager::OnAssetAdded);
 	AssetRegistry.OnAssetRemoved().AddRaw(this, &FAssetSearchManager::OnAssetRemoved);
 
+	TArray<FAssetData> AllAssets;
 	AssetRegistry.GetAllAssets(AllAssets, true);
 
 	for (const FAssetData& Data : AllAssets)
 	{
 		OnAssetAdded(Data);
 	}
-	
-	PendingDatabaseUpdates++;
-	UpdateOperations.Enqueue([this, AssetsAvailable = MoveTemp(AllAssets)]() mutable {
-		FScopeLock ScopedLock(&SearchDatabaseCS);
-		//UE_LOG(LogAssetSearch, Log, TEXT(""));
-		SearchDatabase.RemoveAssetsNotInThisSet(AssetsAvailable);
-		PendingDatabaseUpdates--;
-	});
+
+	bFinishedOnce = true;
+	return AllAssets;
 }
 
 uint64 FAssetSearchManager::GetTextHash(FStringView PackageRelativeExportPath) const
diff --git a/Engine/Plugins/Editor/AssetSearch/Source/Private/AssetSearchManager.h b/Engine/Plugins/Editor/AssetSearch/Source/Private/AssetSearchManager.h
--- a/Engine/Plugins/Editor/AssetSearch/Source/Private/AssetSearchManager.h
+++ b/Engine/Plugins/Editor/AssetSearch/Source/Private/AssetSearchManager.h
@@ -55,7 +55,8 @@ private:
 
 	void OnAssetAdded(const FAssetData& InAssetData);
 	void OnAssetRemoved(const FAssetData& InAssetData);
-	void OnAssetScanFinished();
+	void StartScanningAssetsOnAssetScanFinished();
+	TArray<FAssetData> StartScanningAssetsOnAssetScanFinishedInternal();
 
 	void HandleOnGetExtraObjectTags(FAssetRegistryTagsContext Context);
 	void HandlePackageSaved(const FString& PackageFilename, UPackage* Package, FObjectPostSaveContext ObjectSaveContext);
@@ -91,6 +92,7 @@ private:
 
 private:
 	bool bStarted = false;
+	bool bFinishedOnce = false;
 
 private:
 	mutable FCriticalSection AssetSearchManagerCS;
